Takes: Remove unused TakeViewer

diff --git a/Engine/Plugins/VirtualProduction/Takes/Source/TakeViewer/Private/TakeViewerModule.cpp b/Engine/Plugins/VirtualProduction/Takes/Source/TakeViewer/Private/TakeViewerModule.cpp
deleted file mode 100644
diff --git a/Engine/Plugins/VirtualProduction/Takes/Source/TakeViewer/TakeViewer.Build.cs b/Engine/Plugins/VirtualProduction/Takes/Source/TakeViewer/TakeViewer.Build.cs
deleted file mode 100644
diff --git a/Engine/Plugins/VirtualProduction/Takes/Takes.uplugin b/Engine/Plugins/VirtualProduction/Takes/Takes.uplugin
--- a/Engine/Plugins/VirtualProduction/Takes/Takes.uplugin
+++ b/Engine/Plugins/VirtualProduction/Takes/Takes.uplugin
@@ -42,10 +42,5 @@
 			"Type" : "Developer",
 			"LoadingPhase" : "PostEngineInit"
 		},
-		{
-			"Name" : "TakeViewer",
-			"Type" : "Developer",
-			"LoadingPhase" : "PostEngineInit"
-		}
 	]
 }
\ No newline at end of file
