Take Recorder: Frame rate tweaks

Show the actual meta data frame rate in the cockpit.
Remove per actor sample rates - record directly with the sample rate of the movie scene.

diff --git a/Engine/Plugins/VirtualProduction/Takes/Source/TakeRecorder/Private/Widgets/STakeRecorderCockpit.cpp b/Engine/Plugins/VirtualProduction/Takes/Source/TakeRecorder/Private/Widgets/STakeRecorderCockpit.cpp
--- a/Engine/Plugins/VirtualProduction/Takes/Source/TakeRecorder/Private/Widgets/STakeRecorderCockpit.cpp
+++ b/Engine/Plugins/VirtualProduction/Takes/Source/TakeRecorder/Private/Widgets/STakeRecorderCockpit.cpp
@@ -368,7 +368,7 @@ void STakeRecorderCockpit::Construct(const FArguments& InArgs)
 						SNew(STextBlock)
 						.ColorAndOpacity(FSlateColor::UseSubduedForeground())
 						.Font(FTakeRecorderStyle::Get().GetFontStyle("TakeRecorder.Cockpit.SmallText"))
-						.Text(FText::FromString(TEXT("60 fps")))
+						.Text(this, &STakeRecorderCockpit::GetFrameRateText)
 					]
 				]
 
@@ -598,6 +598,11 @@ FText STakeRecorderCockpit::GetTimestampText() const
 	return TextTime;
 }
 
+FText STakeRecorderCockpit::GetFrameRateText() const
+{
+	return GetFrameRate().ToPrettyText();
+}
+
 FFrameRate STakeRecorderCockpit::GetFrameRate() const
 {
 	return TakeMetaData->GetFrameRate();
@@ -637,22 +642,6 @@ void STakeRecorderCockpit::SetUserDescriptionText(const FText& InNewText, ETextC
 	}
 }
 
-void STakeRecorderCockpit::SetFrameRate(FFrameRate InNewFrameRate)
-{
-	FScopedTransaction Transaction(LOCTEXT("SetFrameRate_Transaction", "Set Frame Rate"));
-	TakeMetaData->Modify();
-
-	TakeMetaData->SetFrameRate(InNewFrameRate);
-
-	ULevelSequence* Sequence   = LevelSequenceAttribute.Get();
-	UMovieScene*    MovieScene = Sequence ? Sequence->GetMovieScene() : nullptr;
-	if (MovieScene)
-	{
-		MovieScene->Modify();
-		MovieScene->SetDisplayRate(InNewFrameRate);
-	}
-}
-
 void STakeRecorderCockpit::SetDurationText(const FText& InNewText, ETextCommit::Type)
 {
 	double CurrentFrameTime = TakeMetaData->GetDuration().AsDecimal();
diff --git a/Engine/Plugins/VirtualProduction/Takes/Source/TakeRecorder/Private/Widgets/STakeRecorderCockpit.h b/Engine/Plugins/VirtualProduction/Takes/Source/TakeRecorder/Private/Widgets/STakeRecorderCockpit.h
--- a/Engine/Plugins/VirtualProduction/Takes/Source/TakeRecorder/Private/Widgets/STakeRecorderCockpit.h
+++ b/Engine/Plugins/VirtualProduction/Takes/Source/TakeRecorder/Private/Widgets/STakeRecorderCockpit.h
@@ -62,7 +62,7 @@ private:
 	FText GetTakeWarningText() const;
 
 	EVisibility GetRecordErrorVisibility() const;
-	FText GetRecordErrorText() const;
+	FText GetRecordErrorText() const; 
 
 	void UpdateRecordError();
 	void UpdateTakeError();
@@ -76,9 +76,9 @@ private:
 	FText GetUserDescriptionText() const;
 	void SetUserDescriptionText(const FText& InNewText, ETextCommit::Type);
 
+	FText GetFrameRateText() const;
 	FFrameRate GetFrameRate() const;
 	bool IsFrameRateCompatible(FFrameRate InFrameRate) const;
-	void SetFrameRate(FFrameRate InNewFrameRate);
 
 	FText GetSlateText() const;
 	void SetSlateText(const FText& InNewText, ETextCommit::Type InCommitType);
diff --git a/Engine/Plugins/VirtualProduction/Takes/Source/TakeTrackRecorders/Private/MovieSceneAnimationTrackRecorder.cpp b/Engine/Plugins/VirtualProduction/Takes/Source/TakeTrackRecorders/Private/MovieSceneAnimationTrackRecorder.cpp
--- a/Engine/Plugins/VirtualProduction/Takes/Source/TakeTrackRecorders/Private/MovieSceneAnimationTrackRecorder.cpp
+++ b/Engine/Plugins/VirtualProduction/Takes/Source/TakeTrackRecorders/Private/MovieSceneAnimationTrackRecorder.cpp
@@ -102,11 +102,13 @@ void UMovieSceneAnimationTrackRecorder::CreateTrackImpl()
 
 		if (AnimSequence.IsValid())
 		{
+			FFrameRate SampleRate = MovieScene->GetDisplayRate();
+
 			FText Error;
 			FString Name = SkeletalMeshComponent->GetName();
 			FName SerializedType("Animation");
 			FString FileName = FString::Printf(TEXT("%s_%s"), *(SerializedType.ToString()), *Name);
-			float IntervalTime = AnimSettings->SampleRate.AsDecimal() > 0.0f ? 1.0f / AnimSettings->SampleRate.AsDecimal() : 1.0f / FAnimationRecordingSettings::DefaultSampleRate;
+			float IntervalTime = SampleRate.AsDecimal() > 0.0f ? 1.0f / SampleRate.AsDecimal() : 1.0f / FAnimationRecordingSettings::DefaultSampleRate;
 			FAnimationFileHeader Header(SerializedType, ObjectGuid, IntervalTime);
 
 			USkeleton* AnimSkeleton = AnimSequence->GetSkeleton();
@@ -224,9 +226,12 @@ void UMovieSceneAnimationTrackRecorder::RecordSampleImpl(const FQualifiedFrameTi
 			// We capture world space transforms for actors if they're attached, but we're not recording the attachment parent
 			bRecordInWorldSpace = !OwningTakeRecorderSource->IsOtherActorBeingRecorded(AttachParent->GetOwner());
 		}
+
+		FFrameRate SampleRate = MovieSceneSection->GetTypedOuter<UMovieScene>()->GetDisplayRate();
+
 		//Set this up here so we know that it's parent sources have also been added so we record in the correct space
 		FAnimationRecordingSettings RecordingSettings;
-		RecordingSettings.SampleRate = AnimSettings->SampleRate.AsDecimal();
+		RecordingSettings.SampleRate = SampleRate.AsDecimal();
 		RecordingSettings.InterpMode = AnimSettings->InterpMode;
 		RecordingSettings.TangentMode = AnimSettings->TangentMode;
 		RecordingSettings.Length = 0;
diff --git a/Engine/Plugins/VirtualProduction/Takes/Source/TakeTrackRecorders/Public/TrackRecorders/MovieSceneAnimationTrackRecorderSettings.h b/Engine/Plugins/VirtualProduction/Takes/Source/TakeTrackRecorders/Public/TrackRecorders/MovieSceneAnimationTrackRecorderSettings.h
--- a/Engine/Plugins/VirtualProduction/Takes/Source/TakeTrackRecorders/Public/TrackRecorders/MovieSceneAnimationTrackRecorderSettings.h
+++ b/Engine/Plugins/VirtualProduction/Takes/Source/TakeTrackRecorders/Public/TrackRecorders/MovieSceneAnimationTrackRecorderSettings.h
@@ -17,7 +17,6 @@ public:
 		: Super(ObjInit)
 	, AnimationTrackName(NSLOCTEXT("UMovieSceneAnimationTrackRecorderSettings", "DefaultAnimationTrackName", "RecordedAnimation"))
 	, AnimationSubDirectory(TEXT("Animation"))
-	, SampleRate(30, 1)
 	, bRemoveRootAnimation(true)
 	{
 	}
@@ -30,10 +29,6 @@ public:
 	UPROPERTY(config, EditAnywhere, BlueprintReadWrite, Category = "Animation Recorder Settings")
 	FString AnimationSubDirectory;
 	
-	/** Maximum sample rate of the recorded animation */
-	UPROPERTY(EditAnywhere, Category = "Animation Recorder Settings")
-	FFrameRate SampleRate;
-	
 	/** Interpolation mode for the recorded keys. */
 	UPROPERTY(EditAnywhere, Category = "Animation Recorder Settings", DisplayName = "Interpolation Mode")
 	TEnumAsByte<ERichCurveInterpMode> InterpMode;
