Take Recorder: Split TakeSaveDir to RootTakeSaveDir and TakeSaveDir. The RootTakeSaveDir defaults to "/Game/Cinematics/Takes" and TakeSaveDir defaults to "{year}-{month}-{day}/{slate}_{take}".

diff --git a/Engine/Plugins/VirtualProduction/Takes/Source/TakeRecorder/Private/Recorder/TakeRecorder.cpp b/Engine/Plugins/VirtualProduction/Takes/Source/TakeRecorder/Private/Recorder/TakeRecorder.cpp
--- a/Engine/Plugins/VirtualProduction/Takes/Source/TakeRecorder/Private/Recorder/TakeRecorder.cpp
+++ b/Engine/Plugins/VirtualProduction/Takes/Source/TakeRecorder/Private/Recorder/TakeRecorder.cpp
@@ -384,7 +384,7 @@ bool UTakeRecorder::Initialize(ULevelSequence* LevelSequenceBase, UTakeRecorderS
 		return false;
 	}
 
-	if (!CreateDestinationAsset(*InParameters.Project.TakeSaveDir.Path, LevelSequenceBase, Sources, MetaData, OutError))
+	if (!CreateDestinationAsset(*InParameters.Project.GetTakeAssetPath(), LevelSequenceBase, Sources, MetaData, OutError))
 	{
 		return false;
 	}
@@ -654,7 +654,7 @@ void UTakeRecorder::Start()
 	MovieSceneSerializationNamespace::bAutoSerialize = Parameters.User.bAutoSerialize;
 	if (Parameters.User.bAutoSerialize)
 	{
-		FString AssetName = AssetMetaData->GenerateAssetPath(Parameters.Project.TakeSaveDir.Path);
+		FString AssetName = AssetMetaData->GenerateAssetPath(Parameters.Project.GetTakeAssetPath());
 		FString AssetPath = FPaths::ProjectSavedDir() + AssetName;
 		FPaths::RemoveDuplicateSlashes(AssetPath);
 		IPlatformFile& PlatformFile = FPlatformFileManager::Get().GetPlatformFile();
diff --git a/Engine/Plugins/VirtualProduction/Takes/Source/TakeRecorder/Private/TakeRecorderModule.cpp b/Engine/Plugins/VirtualProduction/Takes/Source/TakeRecorder/Private/TakeRecorderModule.cpp
--- a/Engine/Plugins/VirtualProduction/Takes/Source/TakeRecorder/Private/TakeRecorderModule.cpp
+++ b/Engine/Plugins/VirtualProduction/Takes/Source/TakeRecorder/Private/TakeRecorderModule.cpp
@@ -101,7 +101,7 @@ namespace
 
 		NewTab->SetContent( NewBrowser );
 
-		FString TakesDir = FPaths::GetPath(FPaths::GetPath(GetDefault<UTakeRecorderProjectSettings>()->Settings.TakeSaveDir.Path));
+		FString TakesDir = FPaths::GetPath(FPaths::GetPath(GetDefault<UTakeRecorderProjectSettings>()->Settings.GetTakeAssetPath()));
 		TArray<FString> TakesFolders;
 		TakesFolders.Push(TakesDir);
 		ContentBrowser.SyncBrowserToFolders(TakesFolders, true, false, FTakeRecorderModule::TakesBrowserInstanceName);
diff --git a/Engine/Plugins/VirtualProduction/Takes/Source/TakeRecorder/Private/TakeRecorderSettings.cpp b/Engine/Plugins/VirtualProduction/Takes/Source/TakeRecorder/Private/TakeRecorderSettings.cpp
--- a/Engine/Plugins/VirtualProduction/Takes/Source/TakeRecorder/Private/TakeRecorderSettings.cpp
+++ b/Engine/Plugins/VirtualProduction/Takes/Source/TakeRecorder/Private/TakeRecorderSettings.cpp
@@ -29,7 +29,8 @@ void UTakeRecorderUserSettings::PostEditChangeProperty(FPropertyChangedEvent& Pr
 
 UTakeRecorderProjectSettings::UTakeRecorderProjectSettings()
 {
-	Settings.TakeSaveDir.Path = TEXT("/Game/Cinematics/Takes/{year}-{month}-{day}/{slate}_{take}");
+	Settings.RootTakeSaveDir.Path = TEXT("/Game/Cinematics/Takes");
+	Settings.TakeSaveDir = TEXT("{year}-{month}-{day}/{slate}_{take}");
 	Settings.DefaultSlate = TEXT("Scene_1");
 	Settings.bRecordSourcesIntoSubSequences = true;
 	Settings.bRecordToPossessable = false;
diff --git a/Engine/Plugins/VirtualProduction/Takes/Source/TakeRecorder/Private/Widgets/STakeRecorderCockpit.cpp b/Engine/Plugins/VirtualProduction/Takes/Source/TakeRecorder/Private/Widgets/STakeRecorderCockpit.cpp
--- a/Engine/Plugins/VirtualProduction/Takes/Source/TakeRecorder/Private/Widgets/STakeRecorderCockpit.cpp
+++ b/Engine/Plugins/VirtualProduction/Takes/Source/TakeRecorder/Private/Widgets/STakeRecorderCockpit.cpp
@@ -459,7 +459,7 @@ void STakeRecorderCockpit::UpdateRecordError()
 		return;
 	}
 
-	FString PackageName = TakeMetaData->GenerateAssetPath(GetDefault<UTakeRecorderProjectSettings>()->Settings.TakeSaveDir.Path);
+	FString PackageName = TakeMetaData->GenerateAssetPath(GetDefault<UTakeRecorderProjectSettings>()->Settings.GetTakeAssetPath());
 	if (!FPackageName::IsValidLongPackageName(PackageName))
 	{
 		RecordErrorText = FText::Format(LOCTEXT("ErrorWidget_InvalidPath", "{0} is not a valid asset path."), FText::FromString(PackageName));
diff --git a/Engine/Plugins/VirtualProduction/Takes/Source/TakeRecorder/Private/Widgets/STakeRecorderPanel.cpp b/Engine/Plugins/VirtualProduction/Takes/Source/TakeRecorder/Private/Widgets/STakeRecorderPanel.cpp
--- a/Engine/Plugins/VirtualProduction/Takes/Source/TakeRecorder/Private/Widgets/STakeRecorderPanel.cpp
+++ b/Engine/Plugins/VirtualProduction/Takes/Source/TakeRecorder/Private/Widgets/STakeRecorderPanel.cpp
@@ -861,7 +861,7 @@ void STakeRecorderPanel::ToggleTakeBrowserCheckState(ECheckBoxState CheckState)
 		FAssetRegistryModule& AssetRegistryModule = FModuleManager::LoadModuleChecked<FAssetRegistryModule>(AssetRegistryConstants::ModuleName);
 		IAssetRegistry& AssetRegistry = AssetRegistryModule.Get();
 
-		FString TakesPath = GetTakeMetaData()->GenerateAssetPath(GetDefault<UTakeRecorderProjectSettings>()->Settings.TakeSaveDir.Path);
+		FString TakesPath = GetTakeMetaData()->GenerateAssetPath(GetDefault<UTakeRecorderProjectSettings>()->Settings.GetTakeAssetPath());
 		TakesPath = FPaths::GetPath(*TakesPath);
 
 		while(!TakesPath.IsEmpty())
diff --git a/Engine/Plugins/VirtualProduction/Takes/Source/TakeRecorder/Public/Recorder/TakeRecorderParameters.h b/Engine/Plugins/VirtualProduction/Takes/Source/TakeRecorder/Public/Recorder/TakeRecorderParameters.h
--- a/Engine/Plugins/VirtualProduction/Takes/Source/TakeRecorder/Public/Recorder/TakeRecorderParameters.h
+++ b/Engine/Plugins/VirtualProduction/Takes/Source/TakeRecorder/Public/Recorder/TakeRecorderParameters.h
@@ -57,8 +57,15 @@ struct FTakeRecorderProjectParameters
 
 	TAKERECORDER_API FTakeRecorderProjectParameters();
 
+	/** The take asset path, composed of the TakeRootSaveDir and the TakeSaveDir */
+	TAKERECORDER_API FString GetTakeAssetPath() const { return RootTakeSaveDir.Path / TakeSaveDir; }
+
+	/** The root of the directory in which to save recorded takes. */
+	UPROPERTY(config, EditAnywhere, BlueprintReadWrite, Category = "Take Recorder", meta = (ContentDir))
+	FDirectoryPath RootTakeSaveDir;
+
 	/**
-	 * The location in which to save recorded takes. Supports any of the following format specifiers that will be substituted when a take is recorded:
+	 * The name of the directory in which to save recorded takes. Supports any of the following format specifiers that will be substituted when a take is recorded:
 	 * {day}       - The day of the timestamp for the start of the recording.
 	 * {month}     - The month of the timestamp for the start of the recording.
 	 * {year}      - The year of the timestamp for the start of the recording.
@@ -68,8 +75,8 @@ struct FTakeRecorderProjectParameters
 	 * {take}      - The take number.
 	 * {slate}     - The slate string.
 	 */
-	UPROPERTY(config, EditAnywhere, BlueprintReadWrite, Category="Take Recorder")
-	FDirectoryPath TakeSaveDir;
+	UPROPERTY(config, EditAnywhere, BlueprintReadWrite, Category = "Take Recorder")
+	FString TakeSaveDir;
 
 	/**
 	 * The default name to use for the Slate information
