Take Recorder: Add reason for invalid asset path

diff --git a/Engine/Plugins/VirtualProduction/Takes/Source/TakeRecorder/Private/Widgets/STakeRecorderCockpit.cpp b/Engine/Plugins/VirtualProduction/Takes/Source/TakeRecorder/Private/Widgets/STakeRecorderCockpit.cpp
--- a/Engine/Plugins/VirtualProduction/Takes/Source/TakeRecorder/Private/Widgets/STakeRecorderCockpit.cpp
+++ b/Engine/Plugins/VirtualProduction/Takes/Source/TakeRecorder/Private/Widgets/STakeRecorderCockpit.cpp
@@ -460,9 +460,10 @@ void STakeRecorderCockpit::UpdateRecordError()
 	}
 
 	FString PackageName = TakeMetaData->GenerateAssetPath(GetDefault<UTakeRecorderProjectSettings>()->Settings.GetTakeAssetPath());
-	if (!FPackageName::IsValidLongPackageName(PackageName))
+	FText OutReason;
+	if (!FPackageName::IsValidLongPackageName(PackageName, false, &OutReason))
 	{
-		RecordErrorText = FText::Format(LOCTEXT("ErrorWidget_InvalidPath", "{0} is not a valid asset path."), FText::FromString(PackageName));
+		RecordErrorText = FText::Format(LOCTEXT("ErrorWidget_InvalidPath", "{0} is not a valid asset path. {1}"), FText::FromString(PackageName), OutReason);
 		return;
 	}
 }
