Take Recorder: Remove force linear keys for skeletal animation.

diff --git a/Engine/Plugins/VirtualProduction/Takes/Source/TakeTrackRecorders/Private/MovieScene3DTransformTrackRecorder.cpp b/Engine/Plugins/VirtualProduction/Takes/Source/TakeTrackRecorders/Private/MovieScene3DTransformTrackRecorder.cpp
--- a/Engine/Plugins/VirtualProduction/Takes/Source/TakeTrackRecorders/Private/MovieScene3DTransformTrackRecorder.cpp
+++ b/Engine/Plugins/VirtualProduction/Takes/Source/TakeTrackRecorders/Private/MovieScene3DTransformTrackRecorder.cpp
@@ -449,8 +449,6 @@ void UMovieScene3DTransformTrackRecorder::PostProcessAnimationData(UMovieSceneAn
 	UMovieSceneAnimationTrackRecorderSettings* AnimSettings = CastChecked<UMovieSceneAnimationTrackRecorderSettings>(AnimTrackRecorder->GetTrackRecorderSettings());
 	if (AnimSettings->bRemoveRootAnimation)
 	{
-		// Override the interpolation mode to use linear interpolation to avoid foot sliding
-		InterpolationMode = RCIM_Linear;
 		//Get All Animation Keys
 		FBufferedTransformKeys  AnimationKeys;
 
