Split Prerecord from Startrecord

Multiple fixes done here. First we create a proper PreRecord step that set's everything up, create's tracks, create subscenes, creates files, as needed. We then, due to quirks with asset creation, perform a few ticks and then start recording, which set's up the stored timecodes on the scenes and sections so everything is synced up.

Second fix is to have the animaton recording happen directly as part of the take recorder, instead of handles with it's own tickable editor object
which run's on own tick (different than the take recorder game object tick).  We also automatically get timecode synced time or engine time, avoid duplicate keys etc..

As part of this also fixed for incorreclty rotated animation recording by removing any changes to the root bone animation, we now simply just grab it and move it over to the 3d transform, don't need the inverse since we force root lock from (zero). Once that's
in everything look good except for some skeletal meshes with no existing keys, and we fix that by making sure we set the DefaulTransform based upon the first transform on the root.

Note this change also fixes the issue with his procedural camera's not recording correctly at beginning.

diff --git a/Engine/Plugins/Animation/LiveLink/Source/LiveLinkSequencer/Private/SequencerRecorderSections/MovieSceneLiveLinkTrackRecorder.cpp b/Engine/Plugins/Animation/LiveLink/Source/LiveLinkSequencer/Private/SequencerRecorderSections/MovieSceneLiveLinkTrackRecorder.cpp
--- a/Engine/Plugins/Animation/LiveLink/Source/LiveLinkSequencer/Private/SequencerRecorderSections/MovieSceneLiveLinkTrackRecorder.cpp
+++ b/Engine/Plugins/Animation/LiveLink/Source/LiveLinkSequencer/Private/SequencerRecorderSections/MovieSceneLiveLinkTrackRecorder.cpp
@@ -192,13 +192,12 @@ void UMovieSceneLiveLinkTrackRecorder::RecordSampleImpl(const FQualifiedFrameTim
 			{
 				FFrameNumber FrameNumber;
 
-				if (bSyncedOrForced && GEngine)
+				if (bSyncedOrForced && GEngine && GEngine->GetTimecodeProvider() && GEngine->GetTimecodeProvider()->GetSynchronizationState() == ETimecodeProviderSynchronizationState::Synchronized)
 				{
 					//Get StartTime on Section in TimeCode FrameRate
 					//Convert that to LiveLink FrameRate and subtract out from LiveLink Frame to get section starting from zero.
 					//Finally convert that to the actual MovieScene Section FrameRate(TickResolution).
-					const UTimecodeProvider* TimecodeProvider = GEngine->GetTimecodeProvider();
-					const FQualifiedFrameTime TimeProviderStartFrameTime = FQualifiedFrameTime(MovieSceneSection->TimecodeSource.Timecode, TimecodeProvider->GetFrameRate());
+					const FQualifiedFrameTime TimeProviderStartFrameTime = FQualifiedFrameTime(MovieSceneSection->TimecodeSource.Timecode, FApp::GetTimecodeFrameRate());
 					FQualifiedFrameTime LiveLinkFrameTime = Frame.GetBaseData()->MetaData.SceneTime;
 					const FFrameNumber FrameNumberStart = TimeProviderStartFrameTime.ConvertTo(LiveLinkFrameTime.Rate).FrameNumber;
 					LiveLinkFrameTime.Time.FrameNumber -= FrameNumberStart;
diff --git a/Engine/Plugins/VirtualProduction/Takes/Source/TakeRecorder/Private/Recorder/TakeRecorder.cpp b/Engine/Plugins/VirtualProduction/Takes/Source/TakeRecorder/Private/Recorder/TakeRecorder.cpp
--- a/Engine/Plugins/VirtualProduction/Takes/Source/TakeRecorder/Private/Recorder/TakeRecorder.cpp
+++ b/Engine/Plugins/VirtualProduction/Takes/Source/TakeRecorder/Private/Recorder/TakeRecorder.cpp
@@ -598,63 +598,67 @@ UWorld* UTakeRecorder::GetWorld() const
 
 void UTakeRecorder::Tick(float DeltaTime)
 {
+	FTimecode TimecodeSource = FApp::GetTimecode();
+
 	if (State == ETakeRecorderState::CountingDown)
 	{
+		NumberOfTicksAfterPre = 0;
 		CountdownSeconds = FMath::Max(0.f, CountdownSeconds - DeltaTime);
 		if (CountdownSeconds > 0.f)
 		{
 			return;
 		}
-		Start();
+		PreRecord();
 	}
-	else if (State == ETakeRecorderState::Started)
+	else if (State == ETakeRecorderState::PreRecord)
 	{
-		UTakeRecorderSources* Sources = SequenceAsset->FindOrAddMetaData<UTakeRecorderSources>();
-		FFrameTime CurrentFrameTime = Sources->TickRecording(SequenceAsset, DeltaTime);
-		TSharedPtr<ISequencer> Sequencer = WeakSequencer.Pin();
-		if (Sequencer.IsValid())
+		if (++NumberOfTicksAfterPre == 2) //seems we need 2 ticks to make sure things are settled
 		{
-			FAnimatedRange Range = Sequencer->GetViewRange();
-			UMovieScene*   MovieScene = SequenceAsset->GetMovieScene();
-			if(MovieScene)
-			{ 
-				FFrameRate FrameRate = MovieScene->GetTickResolution();
-				double CurrentTimeSeconds = FrameRate.AsSeconds(CurrentFrameTime) + 0.5f;
-				CurrentTimeSeconds = CurrentTimeSeconds > Range.GetUpperBoundValue() ? CurrentTimeSeconds: Range.GetUpperBoundValue();
-				TRange<double> NewRange(Range.GetLowerBoundValue(), CurrentTimeSeconds );
-				Sequencer->SetViewRange(NewRange, EViewRangeInterpolation::Immediate);
-			}
+			State = ETakeRecorderState::TickingAfterPre;
 		}
 	}
+	else if (State == ETakeRecorderState::TickingAfterPre)
+	{
+		NumberOfTicksAfterPre = 0;
+		Start(TimecodeSource);
+		InternalTick(TimecodeSource, 0.0f);
+	}
+	else if (State == ETakeRecorderState::Started)
+	{
+		InternalTick(TimecodeSource, DeltaTime);
+	}
 }
 
-void UTakeRecorder::Start()
+void UTakeRecorder::InternalTick(const FTimecode& InTimecodeSource, float DeltaTime)
 {
-	State = ETakeRecorderState::Started;
-
+	UTakeRecorderSources* Sources = SequenceAsset->FindOrAddMetaData<UTakeRecorderSources>();
+	CurrentFrameTime = Sources->TickRecording(SequenceAsset, InTimecodeSource, DeltaTime);
 	TSharedPtr<ISequencer> Sequencer = WeakSequencer.Pin();
 	if (Sequencer.IsValid())
 	{
-		FFrameNumber SequenceStart = MovieScene::DiscreteInclusiveLower(SequenceAsset->GetMovieScene()->GetPlaybackRange());
-		// Discard any entity tokens we have so that restore state does not take effect when we delete any sections that recording will be replacing.
-		Sequencer->DiscardEntityTokens();
+		FAnimatedRange Range = Sequencer->GetViewRange();
 		UMovieScene*   MovieScene = SequenceAsset->GetMovieScene();
 		if (MovieScene)
 		{
-			MovieScene->SetClockSource(EUpdateClockSource::Timecode);
-			Sequencer->ResetTimeController();
+			FFrameRate FrameRate = MovieScene->GetTickResolution();
+			double CurrentTimeSeconds = FrameRate.AsSeconds(CurrentFrameTime) + 0.5f;
+			CurrentTimeSeconds = CurrentTimeSeconds > Range.GetUpperBoundValue() ? CurrentTimeSeconds : Range.GetUpperBoundValue();
+			TRange<double> NewRange(Range.GetLowerBoundValue(), CurrentTimeSeconds);
+			Sequencer->SetViewRange(NewRange, EViewRangeInterpolation::Immediate);
 		}
-		Sequencer->SetPlaybackStatus(EMovieScenePlayerStatus::Playing); //set to pause since we will set time while recording
 	}
+}
+
+void UTakeRecorder::PreRecord()
+{
+	State = ETakeRecorderState::PreRecord;
+
 	UTakeRecorderSources* Sources = SequenceAsset->FindMetaData<UTakeRecorderSources>();
 	check(Sources);
-
 	UTakeMetaData* AssetMetaData = SequenceAsset->FindMetaData<UTakeMetaData>();
-	FDateTime UtcNow = FDateTime::UtcNow();
-	AssetMetaData->SetTimestamp(UtcNow);
 
 	//Set the flag to specify if we should auto save the serialized data or not when recording.
-	
+
 	MovieSceneSerializationNamespace::bAutoSerialize = Parameters.User.bAutoSerialize;
 	if (Parameters.User.bAutoSerialize)
 	{
@@ -682,13 +686,44 @@ void UTakeRecorder::Start()
 
 	Sources->SetRecordToSubSequence(Parameters.Project.bRecordSourcesIntoSubSequences);
 
-	Sources->StartRecording(SequenceAsset, Parameters.User.bSaveRecordedAssets ?  &ManifestSerializer : nullptr);
+	Sources->PreRecording(SequenceAsset, Parameters.User.bAutoSerialize ? &ManifestSerializer : nullptr);
 
 	// Refresh sequencer in case the movie scene data has mutated (ie. existing object bindings removed because they will be recorded again)
+	TSharedPtr<ISequencer> Sequencer = WeakSequencer.Pin();
 	if (Sequencer.IsValid())
 	{
 		Sequencer->RefreshTree();
 	}
+}
+
+void UTakeRecorder::Start(const FTimecode& InTimecodeSource)
+{
+	State = ETakeRecorderState::Started;
+
+	TSharedPtr<ISequencer> Sequencer = WeakSequencer.Pin();
+	if (Sequencer.IsValid())
+	{
+		CurrentFrameTime = FFrameTime(0);
+		FFrameNumber SequenceStart = MovieScene::DiscreteInclusiveLower(SequenceAsset->GetMovieScene()->GetPlaybackRange());
+		// Discard any entity tokens we have so that restore state does not take effect when we delete any sections that recording will be replacing.
+		Sequencer->DiscardEntityTokens();
+		UMovieScene*   MovieScene = SequenceAsset->GetMovieScene();
+		if (MovieScene)
+		{
+			MovieScene->SetClockSource(EUpdateClockSource::Timecode);
+			Sequencer->ResetTimeController();
+		}
+		Sequencer->SetPlaybackStatus(EMovieScenePlayerStatus::Playing);
+	}
+	UTakeRecorderSources* Sources = SequenceAsset->FindMetaData<UTakeRecorderSources>();
+	check(Sources);
+
+	UTakeMetaData* AssetMetaData = SequenceAsset->FindMetaData<UTakeMetaData>();
+	FDateTime UtcNow = FDateTime::UtcNow();
+	AssetMetaData->SetTimestamp(UtcNow);
+
+	Sources->StartRecording(SequenceAsset, InTimecodeSource, Parameters.User.bAutoSerialize ? &ManifestSerializer : nullptr);
+
 	OnRecordingStartedEvent.Broadcast(this);
 
 	UTakeRecorderBlueprintLibrary::OnTakeRecorderStarted();
@@ -730,6 +765,18 @@ void UTakeRecorder::Stop()
 		TakeRecorderSourcesSettings.bSaveRecordedAssets = Parameters.User.bSaveRecordedAssets || GEditor == nullptr;
 		TakeRecorderSourcesSettings.bRemoveRedundantTracks = Parameters.User.bRemoveRedundantTracks;
 
+		UMovieScene*   MovieScene = SequenceAsset->GetMovieScene();
+		if (MovieScene)
+		{
+			TRange<FFrameNumber> Range = MovieScene->GetPlaybackRange();
+			//Set Range to what we recorded instead of that large number, this let's  us reliably set camera cut times.
+			MovieScene->SetPlaybackRange(TRange<FFrameNumber>(Range.GetLowerBoundValue(), CurrentFrameTime.FrameNumber));
+			if (Sequencer)
+			{
+				Sequencer->ResetTimeController();
+			}
+		}
+
 		UTakeRecorderSources* Sources = SequenceAsset->FindMetaData<UTakeRecorderSources>();
 		check(Sources);
 		Sources->StopRecording(SequenceAsset, TakeRecorderSourcesSettings);
diff --git a/Engine/Plugins/VirtualProduction/Takes/Source/TakeRecorder/Public/Recorder/TakeRecorder.h b/Engine/Plugins/VirtualProduction/Takes/Source/TakeRecorder/Public/Recorder/TakeRecorder.h
--- a/Engine/Plugins/VirtualProduction/Takes/Source/TakeRecorder/Public/Recorder/TakeRecorder.h
+++ b/Engine/Plugins/VirtualProduction/Takes/Source/TakeRecorder/Public/Recorder/TakeRecorder.h
@@ -24,6 +24,8 @@ UENUM(BlueprintType)
 enum class ETakeRecorderState : uint8
 {
 	CountingDown,
+	PreRecord,
+	TickingAfterPre,
 	Started,
 	Stopped,
 	Cancelled,
@@ -124,9 +126,14 @@ public:
 private:
 
 	/**
-	 * Called after the countdown to start recording
+	 * Called after the countdown to PreRecord
 	 */
-	void Start();
+	void PreRecord();
+
+	/**
+	 * Called after PreRecord To Start
+	 */
+	void Start(const FTimecode& InTimecodeSource);
 
 	/**
 	 * Ticked by a tickable game object to performe any necessary time-sliced logic
@@ -160,6 +167,9 @@ private:
 
 private:
 
+	/** Called by Tick and Start to make sure we record at start */
+	void InternalTick(const FTimecode& InTimecodeSource, float DeltaTime);
+
 	virtual UWorld* GetWorld() const override;
 
 private:
@@ -170,6 +180,9 @@ private:
 	/** The state of this recorder instance */
 	ETakeRecorderState State;
 
+	/** FFrameTime in MovieScene Resolution we are at*/
+	FFrameTime CurrentFrameTime;
+
 	/** The asset that we should output recorded data into */
 	UPROPERTY(transient)
 	ULevelSequence* SequenceAsset;
@@ -201,6 +214,9 @@ private:
 	/** Sequencer ptr that controls playback of the desination asset during the recording */
 	TWeakPtr<ISequencer> WeakSequencer;
 
+	/** Due a few ticks after the pre so we are set up with asset creation */
+	int32 NumberOfTicksAfterPre;
+
 	friend class FTickableTakeRecorder;
 
 private:
diff --git a/Engine/Plugins/VirtualProduction/Takes/Source/TakeRecorderSources/Private/TakeRecorderMicrophoneAudioSource.cpp b/Engine/Plugins/VirtualProduction/Takes/Source/TakeRecorderSources/Private/TakeRecorderMicrophoneAudioSource.cpp
--- a/Engine/Plugins/VirtualProduction/Takes/Source/TakeRecorderSources/Private/TakeRecorderMicrophoneAudioSource.cpp
+++ b/Engine/Plugins/VirtualProduction/Takes/Source/TakeRecorderSources/Private/TakeRecorderMicrophoneAudioSource.cpp
@@ -109,6 +109,18 @@ TArray<UTakeRecorderSource*> UTakeRecorderMicrophoneAudioSource::PreRecording(UL
 		CachedAudioTrack->SetDisplayName(AudioTrackName);
 	}
 
+	FString PathToRecordTo = FPackageName::GetLongPackagePath(InSequence->GetOutermost()->GetPathName());
+	FString BaseName = InSequence->GetName();
+
+	AudioDirectory;
+	AudioDirectory.Path = PathToRecordTo;
+	if (AudioSubDirectory.Len())
+	{
+		AudioDirectory.Path /= AudioSubDirectory;
+	}
+
+	AssetName = MakeNewAssetName(AudioDirectory.Path, BaseName);
+
 	return TArray<UTakeRecorderSource*>();
 }
 
@@ -125,18 +137,6 @@ void UTakeRecorderMicrophoneAudioSource::StartRecording(const FTimecode& InSecti
 {
 	Super::StartRecording(InSectionStartTimecode, InSectionFirstFrame, InSequence);
 
-	FString PathToRecordTo = FPackageName::GetLongPackagePath(InSequence->GetOutermost()->GetPathName());
-	FString BaseName = InSequence->GetName();
-
-	FDirectoryPath AudioDirectory;
-	AudioDirectory.Path = PathToRecordTo;
-	if (AudioSubDirectory.Len())
-	{
-		AudioDirectory.Path /= AudioSubDirectory;
-	}
-
-	FString AssetName = MakeNewAssetName(AudioDirectory.Path, BaseName);
-
 	ISequenceRecorder& Recorder = FModuleManager::Get().LoadModuleChecked<ISequenceRecorder>("SequenceRecorder");
 
 	FSequenceAudioRecorderSettings AudioSettings;
diff --git a/Engine/Plugins/VirtualProduction/Takes/Source/TakeRecorderSources/Private/TakeRecorderMicrophoneAudioSource.h b/Engine/Plugins/VirtualProduction/Takes/Source/TakeRecorderSources/Private/TakeRecorderMicrophoneAudioSource.h
--- a/Engine/Plugins/VirtualProduction/Takes/Source/TakeRecorderSources/Private/TakeRecorderMicrophoneAudioSource.h
+++ b/Engine/Plugins/VirtualProduction/Takes/Source/TakeRecorderSources/Private/TakeRecorderMicrophoneAudioSource.h
@@ -71,4 +71,8 @@ private:
 	TWeakObjectPtr<class UMovieSceneAudioTrack> CachedAudioTrack;
 
 	TUniquePtr<ISequenceAudioRecorder> AudioRecorder;
+	
+	//Created in PreRecord but used in StartRecording.
+	FDirectoryPath AudioDirectory;
+	FString AssetName;
 };
diff --git a/Engine/Plugins/VirtualProduction/Takes/Source/TakeTrackRecorders/Private/MovieScene3DTransformTrackRecorder.cpp b/Engine/Plugins/VirtualProduction/Takes/Source/TakeTrackRecorders/Private/MovieScene3DTransformTrackRecorder.cpp
--- a/Engine/Plugins/VirtualProduction/Takes/Source/TakeTrackRecorders/Private/MovieScene3DTransformTrackRecorder.cpp
+++ b/Engine/Plugins/VirtualProduction/Takes/Source/TakeTrackRecorders/Private/MovieScene3DTransformTrackRecorder.cpp
@@ -510,14 +510,8 @@ void UMovieScene3DTransformTrackRecorder::PostProcessAnimationData(UMovieSceneAn
 				const FFrameRate TickResolution = MovieSceneSection->GetTypedOuter<UMovieScene>()->GetTickResolution();
 				const FFrameNumber StartTime = MovieSceneSection->GetInclusiveStartFrame();
 
-				// we may need to offset the transform here if the animation was not recorded on the root component
 				FTransform InvComponentTransform = AnimTrackRecorder->GetComponentTransform().Inverse();
-				FTransform InitialRootTransform = AnimTrackRecorder->GetInitialRootTransform();
-				if (DefaultTransform.IsSet())
-				{
-					DefaultTransform = InitialRootTransform * DefaultTransform.GetValue();
 
-				}
 				const FRawAnimSequenceTrack& RawTrack = AnimSequence->GetRawAnimationData()[RootIndex];
 				const int32 KeyCount = FMath::Max(FMath::Max(RawTrack.PosKeys.Num(), RawTrack.RotKeys.Num()), RawTrack.ScaleKeys.Num());
 				for (int32 KeyIndex = 0; KeyIndex < KeyCount; KeyIndex++)
@@ -551,12 +545,12 @@ void UMovieScene3DTransformTrackRecorder::PostProcessAnimationData(UMovieSceneAn
 					}
 
 					FFrameNumber AnimationFrame = (AnimSequence->GetTimeAtFrame(KeyIndex) * TickResolution).FloorToFrame();
-					//TODO Fix this.
-					//this works for tests but not props
-					AnimationKeys.Add(InvComponentTransform * InitialRootTransform * Transform * Relative, StartTime + AnimationFrame);
-
-					//this works for props but not tests
-					//AnimationKeys.Add(InvComponentTransform * Transform * InitialRootTransform * Relative, StartTime + AnimationFrame);
+					FTransform Total = InvComponentTransform * Transform * Relative;
+					AnimationKeys.Add(Total, StartTime + AnimationFrame);
+					if (KeyIndex == 0)
+					{
+						DefaultTransform = Total;
+					}
 				}
 			}
 		}
diff --git a/Engine/Plugins/VirtualProduction/Takes/Source/TakeTrackRecorders/Private/MovieSceneAnimationTrackRecorder.cpp b/Engine/Plugins/VirtualProduction/Takes/Source/TakeTrackRecorders/Private/MovieSceneAnimationTrackRecorder.cpp
--- a/Engine/Plugins/VirtualProduction/Takes/Source/TakeTrackRecorders/Private/MovieSceneAnimationTrackRecorder.cpp
+++ b/Engine/Plugins/VirtualProduction/Takes/Source/TakeTrackRecorders/Private/MovieSceneAnimationTrackRecorder.cpp
@@ -5,13 +5,13 @@
 #include "TakesUtils.h"
 #include "Tracks/MovieSceneSkeletalAnimationTrack.h"
 #include "Sections/MovieSceneSkeletalAnimationSection.h"
-#include "AnimationRecorder.h"
 #include "MovieScene.h"
 #include "AssetRegistryModule.h"
 #include "SequenceRecorderSettings.h"
 #include "Components/SkeletalMeshComponent.h"
 #include "Animation/AnimSequence.h"
-
+#include "Engine.h"
+#include "Engine/TimecodeProvider.h"
 DEFINE_LOG_CATEGORY(AnimationSerialization);
 
 bool FMovieSceneAnimationTrackRecorderFactory::CanRecordObject(UObject* InObjectToRecord) const
@@ -101,7 +101,11 @@ void UMovieSceneAnimationTrackRecorder::CreateTrackImpl()
 
 		if (AnimSequence.IsValid())
 		{
-			FFrameRate SampleRate = MovieScene->GetDisplayRate();
+			//If we are syncing to a timecode provider use that's frame rate as our frame rate since
+			//otherwise use the displayrate.
+			const UTimecodeProvider* TimecodeProvider = GEngine->GetTimecodeProvider();
+			FFrameRate SampleRate = (TimecodeProvider && TimecodeProvider->GetSynchronizationState() == ETimecodeProviderSynchronizationState::Synchronized)
+				? TimecodeProvider->GetFrameRate() : MovieScene->GetDisplayRate();
 
 			FText Error;
 			FString Name = SkeletalMeshComponent->GetName();
@@ -173,12 +177,11 @@ void UMovieSceneAnimationTrackRecorder::StopRecordingImpl()
 	{
 		// Legacy Animation Recorder allowed recording into an animation asset directly and not creating an movie section
 		const bool bShowAnimationAssetCreatedToast = false;
-		InitialRootTransform = FAnimationRecorderManager::Get().GetInitialRootTransform(SkeletalMeshComponent.Get());
-		FAnimationRecorderManager::Get().StopRecordingAnimation(SkeletalMeshComponent.Get(), bShowAnimationAssetCreatedToast);
+		InitialRootTransform = AnimationRecorder.Recorder.Get()->GetInitialRootTransform();
+		AnimationRecorder.FinishRecording(bShowAnimationAssetCreatedToast);
 	}
 }
 
-
 void UMovieSceneAnimationTrackRecorder::FinalizeTrackImpl()
 { 
  	if(MovieSceneSection.IsValid() && AnimSequence.IsValid() && MovieSceneSection->HasStartFrame())
@@ -200,8 +203,11 @@ void UMovieSceneAnimationTrackRecorder::RecordSampleImpl(const FQualifiedFrameTi
 {
 	// The animation recorder does most of the work here
 	//  Note we wait for first tick so that we can make sure all of the attach tracks are set up .
+	float CurrentSeconds = CurrentTime.AsSeconds();
+
 	if (!bAnimationRecorderCreated)
 	{
+		/*
 		//Reset the start times based upon when the animation really starts.
 		if (MovieSceneSection.IsValid())
 		{
@@ -214,13 +220,14 @@ void UMovieSceneAnimationTrackRecorder::RecordSampleImpl(const FQualifiedFrameTi
 			MovieSceneSection->ExpandToFrame(CurrentFrame + FFrameNumber(1));
 			MovieSceneSection->SetStartFrame(TRangeBound<FFrameNumber>::Inclusive(CurrentFrame));
 		}
+		*/
 		bAnimationRecorderCreated = true;
 		AActor* Actor = nullptr;
 		SkeletalMeshComponent = CastChecked<USkeletalMeshComponent>(ObjectToRecord.Get());
 		Actor = SkeletalMeshComponent->GetOwner();
 		USceneComponent* RootComponent = Actor->GetRootComponent();
 		USceneComponent* AttachParent = RootComponent ? RootComponent->GetAttachParent() : nullptr;
-		
+
 		//In Sequence Recorder this would be done via checking if the component was dynamically created, due to changes in how the take recorder handles this, it no longer 
 		//possible so it seems if it's native do root, otherwise, don't.  
 		bool bRemoveRootAnimation = SkeletalMeshComponent->CreationMethod != EComponentCreationMethod::Native ? false : true;
@@ -229,7 +236,7 @@ void UMovieSceneAnimationTrackRecorder::RecordSampleImpl(const FQualifiedFrameTi
 		AnimSettings->bRemoveRootAnimation = bRemoveRootAnimation;
 		//If not removing root we also don't record in world space ( not totally sure if it matters but matching up with Sequence Recorder)
 		bool bRecordInWorldSpace = bRemoveRootAnimation == false ? false : true;
-		
+
 		if (bRecordInWorldSpace && AttachParent && OwningTakeRecorderSource)
 		{
 			// We capture world space transforms for actors if they're attached, but we're not recording the attachment parent
@@ -246,8 +253,19 @@ void UMovieSceneAnimationTrackRecorder::RecordSampleImpl(const FQualifiedFrameTi
 		RecordingSettings.Length = 0;
 		RecordingSettings.bRecordInWorldSpace = bRecordInWorldSpace;
 		RecordingSettings.bRemoveRootAnimation = bRemoveRootAnimation;
-		FAnimationRecorderManager::Get().RecordAnimation(SkeletalMeshComponent.Get(), AnimSequence.Get(), &AnimationSerializer, RecordingSettings);
+
+		AnimationRecorder.Init(SkeletalMeshComponent.Get(), AnimSequence.Get(), &AnimationSerializer, RecordingSettings);
+		AnimationRecorder.BeginRecording();
+	}
+	else
+	{
+		float DeltaTime = CurrentSeconds - PreviousSeconds;
+		AnimationRecorder.Update(DeltaTime);
 	}
+
+	PreviousSeconds = CurrentSeconds;
+
+
 	if (SkeletalMeshComponent.IsValid())
 	{
 		// re-force updates on as gameplay can sometimes turn these back off!
diff --git a/Engine/Plugins/VirtualProduction/Takes/Source/TakeTrackRecorders/Public/TrackRecorders/MovieSceneAnimationTrackRecorder.h b/Engine/Plugins/VirtualProduction/Takes/Source/TakeTrackRecorders/Public/TrackRecorders/MovieSceneAnimationTrackRecorder.h
--- a/Engine/Plugins/VirtualProduction/Takes/Source/TakeTrackRecorders/Public/TrackRecorders/MovieSceneAnimationTrackRecorder.h
+++ b/Engine/Plugins/VirtualProduction/Takes/Source/TakeTrackRecorders/Public/TrackRecorders/MovieSceneAnimationTrackRecorder.h
@@ -10,6 +10,7 @@
 #include "Animation/AnimSequence.h"
 #include "Serializers/MovieSceneAnimationSerialization.h"
 #include "Sections/MovieSceneSkeletalAnimationSection.h"
+#include "AnimationRecorder.h"
 #include "MovieSceneAnimationTrackRecorder.generated.h"
 
 class FMovieScene3DTransformTrackRecorder;
@@ -56,7 +57,7 @@ protected:
 
 public:
 	void RemoveRootMotion();
-	
+
 	UAnimSequence* GetAnimSequence() const { return AnimSequence.Get(); }
 	USkeletalMesh* GetSkeletalMesh() const { return SkeletalMesh.Get(); }
 	USkeletalMeshComponent* GetSkeletalMeshComponent() const { return SkeletalMeshComponent.Get(); }
@@ -82,7 +83,14 @@ private:
 	/** Inverse we are using to zero out root motion */
 	FTransform InitialRootTransform;
 
-	bool bAnimationRecorderCreated; 
+	bool bAnimationRecorderCreated;
+
+	/** Animatinon Recorder */
+	FAnimRecorderInstance AnimationRecorder;
+
+	/** Previous Seconds to calc Delta used by Animation Recorder **/
+	float  PreviousSeconds;
+
 	/**Serializer */
 	FAnimationSerializer AnimationSerializer;
 };
diff --git a/Engine/Plugins/VirtualProduction/Takes/Source/TakesCore/Private/TakeRecorderSources.cpp b/Engine/Plugins/VirtualProduction/Takes/Source/TakesCore/Private/TakeRecorderSources.cpp
--- a/Engine/Plugins/VirtualProduction/Takes/Source/TakesCore/Private/TakeRecorderSources.cpp
+++ b/Engine/Plugins/VirtualProduction/Takes/Source/TakesCore/Private/TakeRecorderSources.cpp
@@ -127,18 +127,18 @@ void UTakeRecorderSources::StartRecordingRecursive(TArray<UTakeRecorderSource*>
 				{
 					SubsceneTrack = CastChecked<UMovieSceneSubTrack>(InMasterSequence->GetMovieScene()->AddMasterTrack(UMovieSceneSubTrack::StaticClass()));
 				}
-				
+
 				// We create a new sub track for every Source so that we can name the Subtrack after the Source instead of just the sections within it.
 				SubsceneTrack->SetDisplayName(FText::FromString(Source->GetSubsceneTrackName(InMasterSequence)));
 				SubsceneTrack->SetColorTint(Source->TrackTint);
-					
+
 				// When we create the Subscene Track we'll make sure a folder is created for it to sort into and add the new Subscene Track as a child of it.
 				if (bCreateSequencerFolders)
 				{
 					UMovieSceneFolder* Folder = AddFolderForSource(Source, InMasterSequence->GetMovieScene());
 					Folder->AddChildMasterTrack(SubsceneTrack);
 				}
-				
+
 				// We initialize the sequence to start at zero and be a 0 frame length section as there is no data in the sections yet.
 				// We'll have to update these sections each frame as the recording progresses so they appear to get longer like normal
 				// tracks do as we record into them.
@@ -203,7 +203,7 @@ void UTakeRecorderSources::StartRecordingRecursive(TArray<UTakeRecorderSource*>
 			if (!bRecordSourcesToSubSequences && bCreateSequencerFolders)
 			{
 				UMovieSceneFolder* Folder = AddFolderForSource(Source, InMasterSequence->GetMovieScene());
-				
+
 				// Different sources can create different kinds of tracks so we allow each source to decide how it gets
 				// represented inside the folder.
 				Source->AddContentsToFolder(Folder);
@@ -215,7 +215,7 @@ void UTakeRecorderSources::StartRecordingRecursive(TArray<UTakeRecorderSource*>
 	{
 		// We don't want to nestle sub-sequences recursively so we always pass the Master Sequence and not the sequence
 		// created for a new source.
-		StartRecordingRecursive(NewSources, InMasterSequence, Timecode,InManifestSerializer);
+		StartRecordingRecursive(NewSources, InMasterSequence, Timecode, InManifestSerializer);
 		SourcesSerialNumber++;
 
 		bool bHasValidTimecodeSource;
@@ -242,12 +242,178 @@ void UTakeRecorderSources::StartRecordingRecursive(TArray<UTakeRecorderSource*>
 	}
 }
 
-void UTakeRecorderSources::StartRecordingSource(TArray<UTakeRecorderSource *> InSources,const FTimecode& CurrentTimecode)
+
+void UTakeRecorderSources::PreRecordingRecursive(TArray<UTakeRecorderSource*> InSources, ULevelSequence* InMasterSequence, TArray<UTakeRecorderSource*>& NewSourcesOut, FManifestSerializer* InManifestSerializer)
+{
+
+	TArray<UTakeRecorderSource*> NewSources;
+
+	// Optionally create a folder in the Sequencer UI that will contain this source. We don't want sub-sequences to have folders
+	// created for their sources as you would end up with a Subscene with one item in it hidden inside of a folder, so instead
+	// only the master sequence gets folders created.
+	const bool bCreateSequencerFolders = true;
+	NewSourcesOut.Append(InSources);
+
+	for (UTakeRecorderSource* Source : InSources)
+	{
+		if (Source->bEnabled)
+		{
+			ULevelSequence* TargetSequence = InMasterSequence;
+
+			// The Sequencer Take system is built around swapping out sub-sequences. If they want to use this system, we create a sub-sequence
+			// for the Source and tell it to write into this sub-sequence instead of the master sequence. We then keep track of which Source
+			// is using which sub-sequence so that we can push the correct sequence for all points of the Source's recording lifecycle.
+			if (bRecordSourcesToSubSequences && Source->SupportsSubscenes())
+			{
+				const FString& SubSequenceTrackName = ObjectTools::SanitizeObjectName(Source->GetSubsceneTrackName(InMasterSequence));
+				const FString& SubSequenceAssetName = ObjectTools::SanitizeObjectName(Source->GetSubsceneAssetName(InMasterSequence));
+
+				TargetSequence = CreateSubSequenceForSource(InMasterSequence, SubSequenceTrackName, SubSequenceAssetName);
+
+				// If there's already a Subscene Track for our sub-sequence we need to remove that track before create a new one. No data is lost in this process as the
+				// sequence that the subscene points to has been copied by CreateSubSequenceForSource so a new track pointed to the new subsequence includes all the old data.
+				const FString SequenceName = FPaths::GetBaseFilename(TargetSequence->GetPathName());
+				UMovieSceneSubTrack* SubsceneTrack = nullptr;
+
+				for (UMovieSceneTrack* Track : InMasterSequence->GetMovieScene()->GetMasterTracks())
+				{
+					if (Track->IsA<UMovieSceneSubTrack>())
+					{
+						if (Track->GetDisplayName().ToString() == SubSequenceTrackName)
+						{
+							SubsceneTrack = CastChecked<UMovieSceneSubTrack>(Track);
+							SubsceneTrack->RemoveAllAnimationData();
+						}
+					}
+				}
+
+				// We need to add the new subsequence to the master sequence immediately so that it shows up in the UI and you can tell that things
+				// are being recorded, otherwise they don't show up until recording stops and then it magically pops in.
+				if (!SubsceneTrack)
+				{
+					SubsceneTrack = CastChecked<UMovieSceneSubTrack>(InMasterSequence->GetMovieScene()->AddMasterTrack(UMovieSceneSubTrack::StaticClass()));
+				}
+
+				// We create a new sub track for every Source so that we can name the Subtrack after the Source instead of just the sections within it.
+				SubsceneTrack->SetDisplayName(FText::FromString(Source->GetSubsceneTrackName(InMasterSequence)));
+				SubsceneTrack->SetColorTint(Source->TrackTint);
+
+				// When we create the Subscene Track we'll make sure a folder is created for it to sort into and add the new Subscene Track as a child of it.
+				if (bCreateSequencerFolders)
+				{
+					UMovieSceneFolder* Folder = AddFolderForSource(Source, InMasterSequence->GetMovieScene());
+					Folder->AddChildMasterTrack(SubsceneTrack);
+				}
+
+				// We initialize the sequence to start at zero and be a 0 frame length section as there is no data in the sections yet.
+				// We'll have to update these sections each frame as the recording progresses so they appear to get longer like normal
+				// tracks do as we record into them.
+				FFrameNumber RecordStartTime = FFrameNumber(0);
+				UMovieSceneSubSection* NewSubSection = SubsceneTrack->AddSequence(TargetSequence, RecordStartTime, 0);
+
+				NewSubSection->SetRowIndex(SubsceneTrack->GetMaxRowIndex() + 1);
+				SubsceneTrack->FixRowIndices();
+
+				ActiveSubSections.Add(NewSubSection);
+				if (InManifestSerializer)
+				{
+					FName SerializedType("SubSequence");
+					FManifestProperty  ManifestProperty(SubSequenceAssetName, SerializedType, FGuid());
+					InManifestSerializer->WriteFrameData(InManifestSerializer->FramesWritten, ManifestProperty);
+
+					FString AssetPath = InManifestSerializer->GetLocalCaptureDir();
+
+					IPlatformFile& PlatformFile = FPlatformFileManager::Get().GetPlatformFile();
+					if (!PlatformFile.DirectoryExists(*AssetPath))
+					{
+						PlatformFile.CreateDirectory(*AssetPath);
+					}
+
+					AssetPath = AssetPath / SubSequenceAssetName;
+					if (!PlatformFile.DirectoryExists(*AssetPath))
+					{
+						PlatformFile.CreateDirectory(*AssetPath);
+					}
+
+					TSharedPtr<FManifestSerializer> NewManifestSerializer = MakeShared<FManifestSerializer>();
+					CreatedManifestSerializers.Add(NewManifestSerializer);
+					InManifestSerializer = NewManifestSerializer.Get();
+
+					InManifestSerializer->SetLocalCaptureDir(AssetPath);
+
+					FManifestFileHeader Header(SubSequenceAssetName, SerializedType, FGuid());
+					FText Error;
+					FString FileName = FString::Printf(TEXT("%s_%s"), *(SerializedType.ToString()), *(SubSequenceAssetName));
+
+					if (!InManifestSerializer->OpenForWrite(FileName, Header, Error))
+					{
+						UE_LOG(SubSequenceSerialization, Warning, TEXT("Error Opening Sequence Sequencer File: Subject '%s' Error '%s'"), *(SubSequenceAssetName), *(Error.ToString()));
+					}
+				}
+			}
+
+			// Update our mappings of which sources use which sub-sequence.
+			SourceSubSequenceMap.FindOrAdd(Source) = TargetSequence;
+
+			for (UTakeRecorderSource* NewlyAddedSource : Source->PreRecording(TargetSequence, InMasterSequence, InManifestSerializer))
+			{
+				// Add it to our classes list of sources 
+				Sources.Add(NewlyAddedSource);
+
+				// And then track it separately so we can recursively call PreRecording 
+				NewSources.Add(NewlyAddedSource);
+			}
+
+			// We need to wait until PreRecording is called on a source before asking it to place itself in a folder
+			// so that the Source has had a chance to create any required sections that will go in the folder.
+			if (!bRecordSourcesToSubSequences && bCreateSequencerFolders)
+			{
+				UMovieSceneFolder* Folder = AddFolderForSource(Source, InMasterSequence->GetMovieScene());
+
+				// Different sources can create different kinds of tracks so we allow each source to decide how it gets
+				// represented inside the folder.
+				Source->AddContentsToFolder(Folder);
+			}
+		}
+	}
+
+	if (NewSources.Num())
+	{
+		// We don't want to nestle sub-sequences recursively so we always pass the Master Sequence and not the sequence
+		// created for a new source.
+		PreRecordingRecursive(NewSources, InMasterSequence, NewSourcesOut, InManifestSerializer);
+		SourcesSerialNumber++;
+	}
+}
+
+void UTakeRecorderSources::StartRecordingPreRecordedSources(const FTimecode& CurrentTimecode)
+{
+	StartRecordingTheseSources(PreRecordedSources, CurrentTimecode);
+	PreRecordedSources.Reset(0);
+}
+
+void UTakeRecorderSources::PreRecordSources(TArray<UTakeRecorderSource *> InSources)
+{
+	PreRecordedSources.Reset(0);
+	PreRecordingRecursive(InSources, CachedLevelSequence, PreRecordedSources, CachedManifestSerializer);
+}
+
+
+void UTakeRecorderSources::StartRecordingSource(TArray<UTakeRecorderSource *> InSources, const FTimecode& CurrentTimecode)
 {
 	// This calls PreRecording recursively on every source so that all sources that get added by another source
 	// have had PreRecording called.
-	StartRecordingRecursive(InSources, CachedLevelSequence,CurrentTimecode, CachedManifestSerializer);
+	TArray<UTakeRecorderSource *> NewSources;
+	PreRecordingRecursive(InSources, CachedLevelSequence, NewSources, CachedManifestSerializer);
+	if (NewSources.Num() > 0)
+	{
+		InSources.Append(NewSources);
+	}
+	StartRecordingTheseSources(InSources, CurrentTimecode);
+}
 
+void UTakeRecorderSources::StartRecordingTheseSources(const TArray<UTakeRecorderSource *>& InSources, const FTimecode& CurrentTimecode)
+{
 	bool bHasValidTimecodeSource;
 	FQualifiedFrameTime QualifiedSequenceTime = GetCurrentRecordingFrameTime(CurrentTimecode, bHasValidTimecodeSource);
 	for (auto Source : InSources)
@@ -255,6 +421,10 @@ void UTakeRecorderSources::StartRecordingSource(TArray<UTakeRecorderSource *> In
 		if (Source->bEnabled)
 		{
 			ULevelSequence* SourceSequence = SourceSubSequenceMap[Source];
+			if (bRecordSourcesToSubSequences && Source->SupportsSubscenes()) //Set Timcode on MovieScene if we created a sub scene for it
+			{
+				SourceSequence->GetMovieScene()->TimecodeSource = CurrentTimecode;
+			}
 			FFrameNumber FrameNumber = QualifiedSequenceTime.ConvertTo(SourceSequence->GetMovieScene()->GetTickResolution()).FloorToFrame();
 			Source->TimecodeSource = CurrentTimecode;
 			if (bHasValidTimecodeSource)
@@ -272,7 +442,8 @@ void UTakeRecorderSources::StartRecordingSource(TArray<UTakeRecorderSource *> In
 	}
 }
 
-void UTakeRecorderSources::StartRecording(class ULevelSequence* InSequence, FManifestSerializer* InManifestSerializer)
+
+void UTakeRecorderSources::PreRecording(class ULevelSequence* InSequence, FManifestSerializer* InManifestSerializer)
 {
 	// We want to cache the Serializer and Level Sequence in case more objects start recording mid-recording.
 	// We want them to use the same logic flow as if initialized from scratch so that they properly sort into
@@ -280,24 +451,29 @@ void UTakeRecorderSources::StartRecording(class ULevelSequence* InSequence, FMan
 	CachedManifestSerializer = InManifestSerializer;
 	CachedLevelSequence = InSequence;
 
+	PreRecordSources(Sources);
+
+}
+
+void UTakeRecorderSources::StartRecording(class ULevelSequence* InSequence, const FTimecode& InTimecodeSource, FManifestSerializer* InManifestSerializer)
+{
+
 	bIsRecording = true;
 	TimeSinceRecordingStarted = 0.f;
 	LastTimecodeFrameNumber.Reset();
 	TargetLevelSequenceTickResolution = InSequence->GetMovieScene()->GetTickResolution();
 
-	FTimecode TimecodeSource = FApp::GetTimecode();
-	InSequence->GetMovieScene()->TimecodeSource = TimecodeSource;
-	StartRecordingTimecodeSource = TimecodeSource;
-	StartRecordingSource(Sources, TimecodeSource);
+	InSequence->GetMovieScene()->TimecodeSource = InTimecodeSource;
+	StartRecordingTimecodeSource = InTimecodeSource;
+	StartRecordingPreRecordedSources(InTimecodeSource);
 }
 
-FFrameTime UTakeRecorderSources::TickRecording(class ULevelSequence* InSequence,float DeltaTime)
+FFrameTime UTakeRecorderSources::TickRecording(class ULevelSequence* InSequence, const FTimecode& InTimecodeSource, float DeltaTime)
 {
-	FTimecode CurrentTimecode = FApp::GetTimecode();
 	bool bHasValidTimecodeSource;
-	FQualifiedFrameTime FrameTime = GetCurrentRecordingFrameTime(CurrentTimecode,bHasValidTimecodeSource);
+	FQualifiedFrameTime FrameTime = GetCurrentRecordingFrameTime(InTimecodeSource, bHasValidTimecodeSource);
 	FQualifiedFrameTime SourceFrameTime(FrameTime);
-	bool bTimeIncremented = DeltaTime > 0.0f;
+	bool bTimeIncremented = (DeltaTime > 0.0f || !LastTimecodeFrameNumber.IsSet());
 	if (bHasValidTimecodeSource)
 	{
 		//We leave this ins timecode frame rate since the sources convert it later (cbb and faster to do it here, we actually do it below
@@ -313,7 +489,6 @@ FFrameTime UTakeRecorderSources::TickRecording(class ULevelSequence* InSequence,
 		}
 		LastTimecodeFrameNumber = SourceFrameTime.Time.FrameNumber;
 	}
-
 	if (bTimeIncremented) //only record if time incremented, may not with timecode providers with low frame rates
 	{
 		for (auto Source : Sources)
@@ -387,10 +562,11 @@ FQualifiedFrameTime UTakeRecorderSources::GetCurrentRecordingFrameTime(const FTi
 		// to determine what frame the data should go on. If the engine is ticking faster than the given
 		// Timecode framerate then there will be multiple frames submitted with the same qualified time
 		// and the data sources will end up only storing the latest call on that frame.
+
 		if (TimecodeProvider && TimecodeProvider->GetSynchronizationState() == ETimecodeProviderSynchronizationState::Synchronized)
 		{
-
-			const FFrameNumber QualifiedFrameNumber = TimecodeProvider->GetTimecode().ToFrameNumber(FApp::GetTimecodeFrameRate());
+			FTimecode Timecode = FApp::GetTimecode();
+			const FFrameNumber QualifiedFrameNumber = Timecode.ToFrameNumber(FApp::GetTimecodeFrameRate());
 			FrameTime = FQualifiedFrameTime(FFrameTime(QualifiedFrameNumber), FApp::GetTimecodeFrameRate());
 
 			bHasValidTimecodeSource = true;
@@ -408,8 +584,8 @@ FQualifiedFrameTime UTakeRecorderSources::GetCurrentRecordingFrameTime(const FTi
 		//Use Level Sequence TickRate to make conversions cleaner later on.
 		const FFrameNumber FrameNumber = TargetLevelSequenceTickResolution.AsFrameNumber(TimeSinceRecordingStarted);
 		FrameTime = FQualifiedFrameTime(FFrameTime(FrameNumber), TargetLevelSequenceTickResolution);
-	}	
-	
+	}
+
 	return FrameTime;
 }
 
@@ -418,6 +594,7 @@ void UTakeRecorderSources::StopRecording(class ULevelSequence* InSequence, FTake
 	bIsRecording = false;
 	TimeSinceRecordingStarted = 0.f;
 	LastTimecodeFrameNumber.Reset();
+
 	for (auto Source : Sources)
 	{
 		if (Source->bEnabled)
@@ -533,7 +710,7 @@ ULevelSequence* UTakeRecorderSources::CreateSubSequenceForSource(ULevelSequence*
 	}
 
 	FString NewPath = FString::Printf(TEXT("%s/%s_Subscenes/%s"), *SequenceDirectory, *SequenceName, *SubSequenceAssetName);
-	
+
 	ULevelSequence* OutAsset = nullptr;
 	TakesUtils::CreateNewAssetPackage<ULevelSequence>(NewPath, OutAsset, nullptr, ExistingSubSequence);
 	if (OutAsset)
@@ -650,7 +827,7 @@ void UTakeRecorderSources::RemoveRedundantTracks()
 				ParentBindings.Add(Possessable->GetParent());
 			}
 		}
-		
+
 		TArray<FGuid> BindingsToRemove;
 		for (const FMovieSceneBinding& Binding : MovieScene->GetBindings())
 		{
@@ -664,7 +841,7 @@ void UTakeRecorderSources::RemoveRedundantTracks()
 		{
 			continue;
 		}
-	
+
 		for (FGuid BindingToRemove : BindingsToRemove)
 		{
 			MovieScene->RemovePossessable(BindingToRemove);
diff --git a/Engine/Plugins/VirtualProduction/Takes/Source/TakesCore/Public/TakeRecorderSources.h b/Engine/Plugins/VirtualProduction/Takes/Source/TakesCore/Public/TakeRecorderSources.h
--- a/Engine/Plugins/VirtualProduction/Takes/Source/TakesCore/Public/TakeRecorderSources.h
+++ b/Engine/Plugins/VirtualProduction/Takes/Source/TakesCore/Public/TakeRecorderSources.h
@@ -13,13 +13,6 @@
 
 class UTakeRecorderSource;
 
-enum class RecordPass
-{
-	PreRecord,
-	StartRecord,
-	StopRecord,
-	PostRecord
-};
 DECLARE_LOG_CATEGORY_EXTERN(SubSequenceSerialization, Verbose, All);
 
 struct FTakeRecorderSourcesSettings
@@ -82,9 +75,9 @@ public:
 	* use TArrayView.
 	* DO NOT MODIFY THIS ARRAY, modifications will be lost.
 	*/
- 	UFUNCTION(BlueprintPure, DisplayName = "Get Sources (Copy)", Category = "Take Recorder")
- 	TArray<UTakeRecorderSource*> GetSourcesCopy() const
- 	{
+	UFUNCTION(BlueprintPure, DisplayName = "Get Sources (Copy)", Category = "Take Recorder")
+	TArray<UTakeRecorderSource*> GetSourcesCopy() const
+	{
 		return TArray<UTakeRecorderSource*>(Sources);
 	}
 	/**
@@ -102,7 +95,7 @@ public:
 
 	/** Calls the recording initialization flows on each of the specified sources. */
 	UFUNCTION(BlueprintCallable, Category = "Take Recorder")
-	void StartRecordingSource(TArray<UTakeRecorderSource*> InSources,const FTimecode& CurrentTiimecode);
+	void StartRecordingSource(TArray<UTakeRecorderSource*> InSources, const FTimecode& CurrentTiimecode);
 public:
 
 	/**
@@ -121,19 +114,25 @@ public:
 	 */
 	void UnbindSourcesChanged(FDelegateHandle Handle);
 
-public: 
+public:
+
+	/*
+	 * Pre recording pass
+	*
+	*/
+	void PreRecording(class ULevelSequence* InSequence, FManifestSerializer* InManifestSerializer);
 
 	/*
 	 * Start recording pass
 	 *
 	 */
-	void StartRecording(class ULevelSequence* InSequence, FManifestSerializer* InManifestSerializer);
+	void StartRecording(class ULevelSequence* InSequence, const FTimecode& InTimecodeSource, FManifestSerializer* InManifestSerializer);
 
 	/*
 	* Tick recording pass
 	* @return Current Frame Number we are recording at.
 	*/
-	FFrameTime TickRecording(class ULevelSequence* InSequence, float DeltaTime);
+	FFrameTime TickRecording(class ULevelSequence* InSequence, const FTimecode& InTimecodeSource, float DeltaTime);
 
 	/*
 	* Stop recording pass
@@ -163,6 +162,9 @@ private:
 	/** Calls PreRecording on sources recursively allowing them to create other sources which properly get PreRecording called on them as well. */
 	void StartRecordingRecursive(TArray<UTakeRecorderSource*> InSources, ULevelSequence* InSequence, const FTimecode& Timecode, FManifestSerializer* InManifestSerializer);
 
+	/** Calls PreRecording on sources recursively allowing them to create other sources which properly get PreRecording called on them as well. */
+	void PreRecordingRecursive(TArray<UTakeRecorderSource*> InSources, ULevelSequence* InMasterSequence, TArray<UTakeRecorderSource*>& NewSourcesOut, FManifestSerializer* InManifestSerializer);
+
 	/** Finds the folder that the given Source should be created in, creating it if necessary. */
 	class UMovieSceneFolder* AddFolderForSource(const UTakeRecorderSource* InSource, class UMovieScene* InMovieScene);
 
@@ -172,6 +174,12 @@ private:
 	/** Remove object bindings that don't have any tracks and are not bindings for attach/path tracks */
 	void RemoveRedundantTracks();
 
+	void StartRecordingPreRecordedSources(const FTimecode& CurrentTimecode);
+
+	void PreRecordSources(TArray<UTakeRecorderSource *> InSources);
+
+	void StartRecordingTheseSources(const TArray<UTakeRecorderSource *>& InSources, const FTimecode& CurrentTimecode);
+
 private:
 
 	/** The array of all sources contained within this list */
@@ -194,7 +202,7 @@ private:
 
 	/** What Tick Resolution is the target level sequence we're recording into? Used to convert seconds into FrameNumbers. */
 	FFrameRate TargetLevelSequenceTickResolution;
-	
+
 	/** Non-serialized serial number that is used for updating UI when the source list changes */
 	uint32 SourcesSerialNumber;
 
@@ -215,5 +223,8 @@ private:
 
 	/** Last Timecode Frame Number, used to avoid recording same time twice*/
 	TOptional<FFrameNumber> LastTimecodeFrameNumber;
-};
 
+	/** All sources after PreRecord */
+	TArray<UTakeRecorderSource *> PreRecordedSources;
+
+};
diff --git a/Engine/Source/Editor/SequenceRecorder/Private/AnimationRecorder.cpp b/Engine/Source/Editor/SequenceRecorder/Private/AnimationRecorder.cpp
--- a/Engine/Source/Editor/SequenceRecorder/Private/AnimationRecorder.cpp
+++ b/Engine/Source/Editor/SequenceRecorder/Private/AnimationRecorder.cpp
@@ -581,7 +581,7 @@ bool FAnimationRecorder::Record(USkeletalMeshComponent* Component, FTransform co
 						}
 						else
 						{
-							InvInitialRootTransform = FTransform::Identity;
+							InitialRootTransform = InvInitialRootTransform = FTransform::Identity;
 						}
 						SkeletonRootIndex = BoneIndex;
 						break;
@@ -608,9 +608,6 @@ bool FAnimationRecorder::Record(USkeletalMeshComponent* Component, FTransform co
 				// if record local to world, we'd like to consider component to world to be in root
 				else
 				{
-					// Remove initial root transform
-					LocalTransform *= InvInitialRootTransform;
-
 					if (bRecordLocalToWorld)
 					{
 						LocalTransform *= ComponentToWorld;
