Fix Linux compile error.

diff --git a/Engine/Plugins/VirtualProduction/Takes/Source/TakeRecorderSources/Private/TakeRecorderMicrophoneAudioSource.cpp b/Engine/Plugins/VirtualProduction/Takes/Source/TakeRecorderSources/Private/TakeRecorderMicrophoneAudioSource.cpp
--- a/Engine/Plugins/VirtualProduction/Takes/Source/TakeRecorderSources/Private/TakeRecorderMicrophoneAudioSource.cpp
+++ b/Engine/Plugins/VirtualProduction/Takes/Source/TakeRecorderSources/Private/TakeRecorderMicrophoneAudioSource.cpp
@@ -112,7 +112,6 @@ TArray<UTakeRecorderSource*> UTakeRecorderMicrophoneAudioSource::PreRecording(UL
 	FString PathToRecordTo = FPackageName::GetLongPackagePath(InSequence->GetOutermost()->GetPathName());
 	FString BaseName = InSequence->GetName();
 
-	AudioDirectory;
 	AudioDirectory.Path = PathToRecordTo;
 	if (AudioSubDirectory.Len())
 	{
