CIS Fix JMonolithic header

diff --git a/Engine/Plugins/VirtualProduction/Takes/Source/TakeTrackRecorders/Private/MovieSceneAnimationTrackRecorder.cpp b/Engine/Plugins/VirtualProduction/Takes/Source/TakeTrackRecorders/Private/MovieSceneAnimationTrackRecorder.cpp
--- a/Engine/Plugins/VirtualProduction/Takes/Source/TakeTrackRecorders/Private/MovieSceneAnimationTrackRecorder.cpp
+++ b/Engine/Plugins/VirtualProduction/Takes/Source/TakeTrackRecorders/Private/MovieSceneAnimationTrackRecorder.cpp
@@ -10,7 +10,6 @@
 #include "SequenceRecorderSettings.h"
 #include "Components/SkeletalMeshComponent.h"
 #include "Animation/AnimSequence.h"
-#include "Engine.h"
 #include "Engine/TimecodeProvider.h"
 DEFINE_LOG_CATEGORY(AnimationSerialization);
 
