CIS FIx pt2-

diff --git a/Engine/Plugins/VirtualProduction/Takes/Source/TakeTrackRecorders/Private/MovieSceneAnimationTrackRecorder.cpp b/Engine/Plugins/VirtualProduction/Takes/Source/TakeTrackRecorders/Private/MovieSceneAnimationTrackRecorder.cpp
--- a/Engine/Plugins/VirtualProduction/Takes/Source/TakeTrackRecorders/Private/MovieSceneAnimationTrackRecorder.cpp
+++ b/Engine/Plugins/VirtualProduction/Takes/Source/TakeTrackRecorders/Private/MovieSceneAnimationTrackRecorder.cpp
@@ -11,6 +11,8 @@
 #include "Components/SkeletalMeshComponent.h"
 #include "Animation/AnimSequence.h"
 #include "Engine/TimecodeProvider.h"
+#include "Engine/Engine.h"
+
 DEFINE_LOG_CATEGORY(AnimationSerialization);
 
 bool FMovieSceneAnimationTrackRecorderFactory::CanRecordObject(UObject* InObjectToRecord) const
