Fix for Weapon Mesh orientation with Player recording.

We got rid of the inverse tranform we used to do in order to try to handle root motion but were still applying that back when dealing with attachments to a skeletal mesh. Got rid of that now.

diff --git a/Engine/Plugins/VirtualProduction/Takes/Source/TakeRecorderSources/Private/TakeRecorderActorSource.cpp b/Engine/Plugins/VirtualProduction/Takes/Source/TakeRecorderSources/Private/TakeRecorderActorSource.cpp
--- a/Engine/Plugins/VirtualProduction/Takes/Source/TakeRecorderSources/Private/TakeRecorderActorSource.cpp
+++ b/Engine/Plugins/VirtualProduction/Takes/Source/TakeRecorderSources/Private/TakeRecorderActorSource.cpp
@@ -1469,30 +1469,6 @@ FGuid UTakeRecorderActorSource::GetRecordedActorGuid(class AActor* OtherActor) c
 	return FGuid();
 }
 
-FTransform UTakeRecorderActorSource::GetRecordedActorAnimationInitialRootTransform(class AActor* OtherActor) const
-{
-	UTakeRecorderSources* OwningSources = CastChecked<UTakeRecorderSources>(GetOuter());
-	for (UTakeRecorderSource* Source : OwningSources->GetSources())
-	{
-		if (UTakeRecorderActorSource* ActorSource = Cast<UTakeRecorderActorSource>(Source))
-		{
-			AActor* OtherTarget = ActorSource->Target.Get();
-			if (OtherTarget && OtherActor && (OtherTarget == OtherActor || OtherTarget->GetName() == OtherActor->GetName()))
-			{
-				for (UMovieSceneTrackRecorder* TrackRecorder : ActorSource->TrackRecorders)
-				{
-					if (TrackRecorder->IsA<UMovieSceneAnimationTrackRecorder>())
-					{
-						return Cast<UMovieSceneAnimationTrackRecorder>(TrackRecorder)->GetInitialRootTransform();
-					}
-				}
-
-			}
-		}
-	}
-	return FTransform::Identity;
-}
-
 FMovieSceneSequenceID UTakeRecorderActorSource::GetLevelSequenceID(class AActor* OtherActor)
 {
 	FMovieSceneSequenceID OutSequenceID = MovieSceneSequenceID::Root;
diff --git a/Engine/Plugins/VirtualProduction/Takes/Source/TakeRecorderSources/Private/TakeRecorderActorSource.h b/Engine/Plugins/VirtualProduction/Takes/Source/TakeRecorderSources/Private/TakeRecorderActorSource.h
--- a/Engine/Plugins/VirtualProduction/Takes/Source/TakeRecorderSources/Private/TakeRecorderActorSource.h
+++ b/Engine/Plugins/VirtualProduction/Takes/Source/TakeRecorderSources/Private/TakeRecorderActorSource.h
@@ -236,8 +236,6 @@ protected:
 	FMovieSceneSequenceID GetLevelSequenceID(class AActor* OtherActor) override;
 	/** Returns generic track recorder settings */
 	FTrackRecorderSettings GetTrackRecorderSettings() const override;
-	/** Returns offset that may get set when recording a skeletal mesh animation. Needed to correctly transform attached children*/
-	FTransform GetRecordedActorAnimationInitialRootTransform(class AActor* OtherActor) const override;
 
 	// ~IMovieSceneTrackRecorderHost Interface
 	
diff --git a/Engine/Plugins/VirtualProduction/Takes/Source/TakeTrackRecorders/Private/MovieScene3DTransformTrackRecorder.cpp b/Engine/Plugins/VirtualProduction/Takes/Source/TakeTrackRecorders/Private/MovieScene3DTransformTrackRecorder.cpp
--- a/Engine/Plugins/VirtualProduction/Takes/Source/TakeTrackRecorders/Private/MovieScene3DTransformTrackRecorder.cpp
+++ b/Engine/Plugins/VirtualProduction/Takes/Source/TakeTrackRecorders/Private/MovieScene3DTransformTrackRecorder.cpp
@@ -112,45 +112,6 @@ void UMovieScene3DTransformTrackRecorder::FinalizeTrackImpl()
  	check (	BufferedTransforms.Times.Num() == BufferedTransforms.ScaleZ.Num());
  
  	SlowTask.EnterProgressFrame();
- 
-	FTransform InvParentAnimationRootTransform = FTransform::Identity;
-	FName SocketName;
-	FName ComponentName;
-	AActor* ActorToRecord = Cast<AActor>(ObjectToRecord.Get());
-	if (ActorToRecord)
-	{
-		AActor* AttachedToActor = SequenceRecorderUtils::GetAttachment(ActorToRecord, SocketName, ComponentName);
-		if (AttachedToActor)
-		{
-			InvParentAnimationRootTransform = OwningTakeRecorderSource->GetRecordedActorAnimationInitialRootTransform(AttachedToActor).Inverse();
-			if (!InvParentAnimationRootTransform.Equals(FTransform::Identity))
-			{
-				if (DefaultTransform.IsSet())
-				{
-					FTransform DT = DefaultTransform.GetValue();
-					DT = DT * InvParentAnimationRootTransform;
-					DefaultTransform = DT;
-					FVector Translation = DT.GetTranslation();
-					FVector EulerRotation = DT.GetRotation().Rotator().Euler();
-					FVector Scale = DT.GetScale3D();
-					TArrayView<FMovieSceneFloatChannel*> FloatChannels = MovieSceneSection->GetChannelProxy().GetChannels<FMovieSceneFloatChannel>();
-					FloatChannels[0]->SetDefault(Translation.X);
-					FloatChannels[1]->SetDefault(Translation.Y);
-					FloatChannels[2]->SetDefault(Translation.Z);
-					FloatChannels[3]->SetDefault(EulerRotation.X);
-					FloatChannels[4]->SetDefault(EulerRotation.Y);
-					FloatChannels[5]->SetDefault(EulerRotation.Z);
-					FloatChannels[6]->SetDefault(Scale.X);
-					FloatChannels[7]->SetDefault(Scale.Y);
-					FloatChannels[8]->SetDefault(Scale.Z);
-				}
-				if (BufferedTransforms.Times.Num() > 0)
-				{
-					BufferedTransforms.PostMultTransform(InvParentAnimationRootTransform);
-				}
-			}
-		}
-	}
 
  	// Try to 're-wind' rotations that look like axis flips
  	// We need to do this as a post-process because the recorder cant reliably access 'wound' rotations:
diff --git a/Engine/Plugins/VirtualProduction/Takes/Source/TakeTrackRecorders/Public/TrackRecorders/IMovieSceneTrackRecorderHost.h b/Engine/Plugins/VirtualProduction/Takes/Source/TakeTrackRecorders/Public/TrackRecorders/IMovieSceneTrackRecorderHost.h
--- a/Engine/Plugins/VirtualProduction/Takes/Source/TakeTrackRecorders/Public/TrackRecorders/IMovieSceneTrackRecorderHost.h
+++ b/Engine/Plugins/VirtualProduction/Takes/Source/TakeTrackRecorders/Public/TrackRecorders/IMovieSceneTrackRecorderHost.h
@@ -168,11 +168,5 @@ public:
 	 */
 	virtual FTrackRecorderSettings GetTrackRecorderSettings() const = 0;
 
-	/** Returns original transform that may get set when recording a skeletal mesh animation. Needed to correctly transform attached children.
-	* @param OtherActor Actor to look for.
-	* @return The FTransform applied to the actor when recording animation, that the attached actor may need to apply to themselves to orient correctly.
-	*/
-	virtual FTransform GetRecordedActorAnimationInitialRootTransform(class AActor* OtherActor) const = 0;
-
 	/**/
 };
