Take Recorder: Advance to next available take number when recording from this

diff --git a/Engine/Plugins/VirtualProduction/Takes/Source/TakeRecorder/Private/Widgets/STakeRecorderCockpit.cpp b/Engine/Plugins/VirtualProduction/Takes/Source/TakeRecorder/Private/Widgets/STakeRecorderCockpit.cpp
--- a/Engine/Plugins/VirtualProduction/Takes/Source/TakeRecorder/Private/Widgets/STakeRecorderCockpit.cpp
+++ b/Engine/Plugins/VirtualProduction/Takes/Source/TakeRecorder/Private/Widgets/STakeRecorderCockpit.cpp
@@ -94,6 +94,15 @@ void STakeRecorderCockpit::Construct(const FArguments& InArgs)
 
 	CacheMetaData();
 
+	if (TakeMetaData && !TakeMetaData->IsLocked())
+	{
+		int32 NextTakeNumber = UTakesCoreBlueprintLibrary::ComputeNextTakeNumber(TakeMetaData->GetSlate());
+		if (NextTakeNumber != TakeMetaData->GetTakeNumber())
+		{
+			TakeMetaData->SetTakeNumber(NextTakeNumber);
+		}
+	}
+
 	UpdateTakeError();
 	UpdateRecordError();
 
