Take Recorder: Disallow recording level script actors

diff --git a/Engine/Plugins/VirtualProduction/Takes/Source/TakeRecorderSources/Private/TakeRecorderSourcesModule.cpp b/Engine/Plugins/VirtualProduction/Takes/Source/TakeRecorderSources/Private/TakeRecorderSourcesModule.cpp
--- a/Engine/Plugins/VirtualProduction/Takes/Source/TakeRecorderSources/Private/TakeRecorderSourcesModule.cpp
+++ b/Engine/Plugins/VirtualProduction/Takes/Source/TakeRecorderSources/Private/TakeRecorderSourcesModule.cpp
@@ -26,6 +26,7 @@
 #include "Framework/MultiBox/MultiBoxExtender.h"
 #include "Widgets/Layout/SBox.h"
 
+#include "Engine/LevelScriptActor.h"
 #include "Engine/Selection.h"
 #include "Editor.h"
 #include "SceneOutlinerModule.h"
@@ -151,7 +152,7 @@ namespace
 			{
 				for (AActor* Actor : Level->Actors)
 				{
-					if (Actor && Actor->IsA(Class))
+					if (Actor && Actor->IsA(Class) && !Actor->IsA(ALevelScriptActor::StaticClass()))
 					{
 						OutActors.AddUnique(Actor);
 					}
