Enabled the Virtual Camera to run in -game. You can't create a new take and view previous takes.

diff --git a/Engine/Plugins/Experimental/VirtualCamera/Source/VirtualCamera/Private/LevelSequencePlaybackController.cpp b/Engine/Plugins/Experimental/VirtualCamera/Source/VirtualCamera/Private/LevelSequencePlaybackController.cpp
--- a/Engine/Plugins/Experimental/VirtualCamera/Source/VirtualCamera/Private/LevelSequencePlaybackController.cpp
+++ b/Engine/Plugins/Experimental/VirtualCamera/Source/VirtualCamera/Private/LevelSequencePlaybackController.cpp
@@ -304,7 +304,7 @@ void ULevelSequencePlaybackController::StopLevelSequencePlay()
 bool ULevelSequencePlaybackController::SetActiveLevelSequence(ULevelSequence* InNewLevelSequence)
 {
 #if WITH_EDITOR
-	if(InNewLevelSequence)
+	if(InNewLevelSequence && !IsRunningGame())
 	{
 		const bool bDoFocusOnEditor = false;
 		FAssetEditorManager::Get().OpenEditorForAsset(InNewLevelSequence);
diff --git a/Engine/Plugins/VirtualProduction/Takes/Source/TakeRecorder/Private/Recorder/TakeRecorderBlueprintLibrary.cpp b/Engine/Plugins/VirtualProduction/Takes/Source/TakeRecorder/Private/Recorder/TakeRecorderBlueprintLibrary.cpp
--- a/Engine/Plugins/VirtualProduction/Takes/Source/TakeRecorder/Private/Recorder/TakeRecorderBlueprintLibrary.cpp
+++ b/Engine/Plugins/VirtualProduction/Takes/Source/TakeRecorder/Private/Recorder/TakeRecorderBlueprintLibrary.cpp
@@ -16,6 +16,16 @@ namespace
 	static UTakeRecorderPanel* CurrentTakeRecorderPanel;
 }
 
+
+bool UTakeRecorderBlueprintLibrary::IsTakeRecorderEnabled()
+{
+#if WITH_EDITOR
+	return !IsRunningGame();
+#else
+	return false;
+#endif 
+}
+
 UTakeRecorder* UTakeRecorderBlueprintLibrary::StartRecording(ULevelSequence* LevelSequence, UTakeRecorderSources* Sources, UTakeMetaData* MetaData, const FTakeRecorderParameters& Parameters)
 {
 	if (IsRecording())
@@ -24,6 +34,22 @@ UTakeRecorder* UTakeRecorderBlueprintLibrary::StartRecording(ULevelSequence* Lev
 		return nullptr;
 	}
 
+	if (LevelSequence == nullptr)
+	{
+		FFrame::KismetExecutionMessage(TEXT("The LevelSequence is invalid."), ELogVerbosity::Error);
+		return nullptr;
+	}
+	if (Sources == nullptr)
+	{
+		FFrame::KismetExecutionMessage(TEXT("The Sources is invalid."), ELogVerbosity::Error);
+		return nullptr;
+	}
+	if (MetaData == nullptr)
+	{
+		FFrame::KismetExecutionMessage(TEXT("The MetaData is invalid."), ELogVerbosity::Error);
+		return nullptr;
+	}
+
 	FText ErrorText = NSLOCTEXT("TakeRecorderBlueprintLibrary", "UnknownError", "An unknown error occurred when trying to start recording");
 
 	UTakeRecorder* NewRecorder = NewObject<UTakeRecorder>(GetTransientPackage(), NAME_None, RF_Transient);
@@ -76,8 +102,15 @@ UTakeRecorderPanel* UTakeRecorderBlueprintLibrary::OpenTakeRecorderPanel()
 		return Existing;
 	}
 
-	FLevelEditorModule& LevelEditorModule = FModuleManager::LoadModuleChecked<FLevelEditorModule>("LevelEditor");
-	LevelEditorModule.GetLevelEditorTabManager()->InvokeTab(ITakeRecorderModule::TakeRecorderTabName);
+	if (!IsRunningGame())
+	{
+		FLevelEditorModule& LevelEditorModule = FModuleManager::LoadModuleChecked<FLevelEditorModule>("LevelEditor");
+		LevelEditorModule.GetLevelEditorTabManager()->InvokeTab(ITakeRecorderModule::TakeRecorderTabName);
+	}
+	else
+	{
+		FFrame::KismetExecutionMessage(TEXT("The Take Recorder Panel will not open because the game is running."), ELogVerbosity::Log);
+	}
 
 	return GetTakeRecorderPanel();
 }
diff --git a/Engine/Plugins/VirtualProduction/Takes/Source/TakeRecorder/Public/Recorder/TakeRecorderBlueprintLibrary.h b/Engine/Plugins/VirtualProduction/Takes/Source/TakeRecorder/Public/Recorder/TakeRecorderBlueprintLibrary.h
--- a/Engine/Plugins/VirtualProduction/Takes/Source/TakeRecorder/Public/Recorder/TakeRecorderBlueprintLibrary.h
+++ b/Engine/Plugins/VirtualProduction/Takes/Source/TakeRecorder/Public/Recorder/TakeRecorderBlueprintLibrary.h
@@ -19,6 +19,12 @@ public:
 	GENERATED_BODY()
 
 
+	/**
+	 * Is the Take Recorder enabled in the build
+	 */
+	UFUNCTION(BlueprintPure, Category="Take Recorder")
+	static bool IsTakeRecorderEnabled();
+
 	/**
 	 * Start a new recording using the specified parameters. Will fail if a recording is currently in progress
 	 *
diff --git a/Engine/Plugins/VirtualProduction/Takes/Source/TakeRecorderSources/Private/TakeRecorderActorSource.cpp b/Engine/Plugins/VirtualProduction/Takes/Source/TakeRecorderSources/Private/TakeRecorderActorSource.cpp
--- a/Engine/Plugins/VirtualProduction/Takes/Source/TakeRecorderSources/Private/TakeRecorderActorSource.cpp
+++ b/Engine/Plugins/VirtualProduction/Takes/Source/TakeRecorderSources/Private/TakeRecorderActorSource.cpp
@@ -47,6 +47,12 @@ static const FName DoNotRecordTag("DoNotRecord");
 
 UTakeRecorderSource* UTakeRecorderActorSource::AddSourceForActor(AActor* InActor, UTakeRecorderSources* InSources)
 {
+	if (InSources == nullptr)
+	{
+		FFrame::KismetExecutionMessage(TEXT("The Source is invalid."), ELogVerbosity::Error);
+		return nullptr;
+	}
+
 	UTakeRecorderActorSource* NewSource = InSources->AddSource<UTakeRecorderActorSource>();
 	NewSource->SetSourceActor(InActor);
 	return NewSource;
