Take Recorder: Autotrack actors that aren't being recorded in the world source.

diff --git a/Engine/Plugins/VirtualProduction/Takes/Source/TakeRecorder/Private/TakeRecorderStyle.cpp b/Engine/Plugins/VirtualProduction/Takes/Source/TakeRecorder/Private/TakeRecorderStyle.cpp
--- a/Engine/Plugins/VirtualProduction/Takes/Source/TakeRecorder/Private/TakeRecorderStyle.cpp
+++ b/Engine/Plugins/VirtualProduction/Takes/Source/TakeRecorder/Private/TakeRecorderStyle.cpp
@@ -155,8 +155,8 @@ FTakeRecorderStyle::FTakeRecorderStyle()
 	Set("ClassIcon.TakeRecorderPlayerSource", new FSlateImageBrush(RootToContentDir(TEXT("TakeRecorderPlayerSource_16x.png")), Icon16x16));
 	Set("ClassThumbnail.TakeRecorderPlayerSource", new FSlateImageBrush(RootToContentDir(TEXT("TakeRecorderPlayerSource_64x.png")), Icon64x64));
 
-	Set("ClassIcon.TakeRecorderWorldSettingsSource", new FSlateImageBrush(RootToContentDir(TEXT("TakeRecorderWorldSettingsSource_16x.png")), Icon16x16));
-	Set("ClassThumbnail.TakeRecorderWorldSettingsSource", new FSlateImageBrush(RootToContentDir(TEXT("TakeRecorderWorldSettingsSource_64x.png")), Icon64x64));
+	Set("ClassIcon.TakeRecorderWorldSource", new FSlateImageBrush(RootToContentDir(TEXT("TakeRecorderWorldSource_16x.png")), Icon16x16));
+	Set("ClassThumbnail.TakeRecorderWorldSource", new FSlateImageBrush(RootToContentDir(TEXT("TakeRecorderWorldSource_64x.png")), Icon64x64));
 
 	Set("TakeRecorder.TabIcon", new FSlateImageBrush(RootToContentDir(TEXT("TabIcon_16x.png")), Icon16x16));
 
diff --git a/Engine/Plugins/VirtualProduction/Takes/Source/TakeRecorderSources/Private/TakeRecorderSourcesModule.cpp b/Engine/Plugins/VirtualProduction/Takes/Source/TakeRecorderSources/Private/TakeRecorderSourcesModule.cpp
--- a/Engine/Plugins/VirtualProduction/Takes/Source/TakeRecorderSources/Private/TakeRecorderSourcesModule.cpp
+++ b/Engine/Plugins/VirtualProduction/Takes/Source/TakeRecorderSources/Private/TakeRecorderSourcesModule.cpp
@@ -33,6 +33,7 @@
 #include "SceneOutlinerPublicTypes.h"
 
 #include "TakeRecorderMicrophoneAudioSource.h"
+#include "TakeRecorderWorldSource.h"
 #include "TrackRecorders/MovieSceneAnimationTrackRecorderSettings.h"
 
 #define LOCTEXT_NAMESPACE "TakeRecorderSources"
@@ -308,6 +309,7 @@ public:
 
 		TakeRecorderModule.RegisterSettingsObject(GetMutableDefault<UTakeRecorderMicrophoneAudioSourceSettings>());
 		TakeRecorderModule.RegisterSettingsObject(GetMutableDefault<UMovieSceneAnimationTrackRecorderEditorSettings>());
+		TakeRecorderModule.RegisterSettingsObject(GetMutableDefault<UTakeRecorderWorldSourceSettings>());
 	}
 
 	virtual void ShutdownModule() override
diff --git a/Engine/Plugins/VirtualProduction/Takes/Source/TakeRecorderSources/Private/TakeRecorderWorldSettingsSource.cpp b/Engine/Plugins/VirtualProduction/Takes/Source/TakeRecorderSources/Private/TakeRecorderWorldSettingsSource.cpp
deleted file mode 100644
diff --git a/Engine/Plugins/VirtualProduction/Takes/Source/TakeRecorderSources/Private/TakeRecorderWorldSettingsSource.h b/Engine/Plugins/VirtualProduction/Takes/Source/TakeRecorderSources/Private/TakeRecorderWorldSettingsSource.h
deleted file mode 100644
diff --git a/Engine/Plugins/VirtualProduction/Takes/Source/TakeRecorderSources/Private/TakeRecorderWorldSource.cpp b/Engine/Plugins/VirtualProduction/Takes/Source/TakeRecorderSources/Private/TakeRecorderWorldSource.cpp
new file mode 100644
--- /dev/null
+++ b/Engine/Plugins/VirtualProduction/Takes/Source/TakeRecorderSources/Private/TakeRecorderWorldSource.cpp
@@ -0,0 +1,199 @@
+// Copyright 1998-2019 Epic Games, Inc. All Rights Reserved.
+
+#include "TakeRecorderWorldSource.h"
+
+#include "TakeRecorderSources.h"
+#include "TakeRecorderActorSource.h"
+#include "TakeRecorderSourcesUtils.h"
+#include "TakesUtils.h"
+
+#include "LevelSequence.h"
+#include "GameFramework/WorldSettings.h"
+
+#include "ISequencer.h"
+#include "MovieSceneFolder.h"
+
+#include "Toolkits/AssetEditorManager.h"
+#include "ILevelSequenceEditorToolkit.h"
+
+UTakeRecorderWorldSourceSettings::UTakeRecorderWorldSourceSettings(const FObjectInitializer& ObjInit)
+	: Super(ObjInit)
+	, bRecordWorldSettings(true)
+	, bAutotrackActors(true)
+{
+}
+
+void UTakeRecorderWorldSourceSettings::PostEditChangeProperty(FPropertyChangedEvent& PropertyChangedEvent)
+{
+	Super::PostEditChangeProperty(PropertyChangedEvent);
+
+	if (HasAnyFlags(RF_ClassDefaultObject))
+	{
+		SaveConfig();
+	}
+}
+
+UTakeRecorderWorldSource::UTakeRecorderWorldSource(const FObjectInitializer& ObjInit)
+	: Super(ObjInit)
+{
+	TrackTint = FColor(129, 129, 129);
+}
+
+TArray<UTakeRecorderSource*> UTakeRecorderWorldSource::PreRecording(class ULevelSequence* InSequence, class ULevelSequence* InMasterSequence, FManifestSerializer* InManifestSerializer)
+{
+	TArray<UTakeRecorderSource*> NewSources;
+
+	UTakeRecorderSources* Sources = InSequence->FindOrAddMetaData<UTakeRecorderSources>();
+
+	// Get the first PIE world's world settings
+	UWorld* World = TakeRecorderSourcesUtils::GetSourceWorld(InSequence);
+
+	if (!World)
+	{
+		return NewSources;
+	}
+
+	if (bRecordWorldSettings)
+	{
+		AWorldSettings* WorldSettings = World ? World->GetWorldSettings() : nullptr;
+
+		if (!WorldSettings)
+		{
+			return NewSources;
+		}
+
+		for (auto Source : Sources->GetSources())
+		{
+			if (Source->IsA<UTakeRecorderActorSource>())
+			{
+				UTakeRecorderActorSource* ActorSource = Cast<UTakeRecorderActorSource>(Source);
+				if (ActorSource->Target.IsValid())
+				{
+					if (ActorSource->Target.Get() == WorldSettings)
+					{
+						return NewSources;
+					}
+				}
+			}
+		}
+
+		UTakeRecorderActorSource* ActorSource = NewObject<UTakeRecorderActorSource>(Sources, UTakeRecorderActorSource::StaticClass(), NAME_None, RF_Transactional);
+		ActorSource->Target = WorldSettings;
+		NewSources.Add(ActorSource);
+		WorldSource = ActorSource;
+	}
+
+	if (bAutotrackActors)
+	{
+		AutotrackActors(InSequence, World);
+	}
+
+	return NewSources;
+}
+
+TArray<UTakeRecorderSource*> UTakeRecorderWorldSource::PostRecording(class ULevelSequence* InSequence, class ULevelSequence* InMasterSequence)
+{
+	TArray<UTakeRecorderSource*> SourcesToRemove;
+
+	if (WorldSource.IsValid())
+	{
+		SourcesToRemove.Add(WorldSource.Get());
+	}
+	return SourcesToRemove;
+}
+
+FText UTakeRecorderWorldSource::GetDisplayTextImpl() const
+{
+	return NSLOCTEXT("UTakeRecorderWorldSource", "Label", "World");
+}
+
+bool UTakeRecorderWorldSource::CanAddSource(UTakeRecorderSources* InSources) const
+{
+	for (UTakeRecorderSource* Source : InSources->GetSources())
+	{
+		if (Source->IsA<UTakeRecorderWorldSource>())
+		{
+			return false;
+		}
+	}
+	return true;
+}
+
+void UTakeRecorderWorldSource::AutotrackActors(class ULevelSequence* InSequence, UWorld* InWorld)
+{
+	if (!InWorld)
+	{
+		return;
+	}
+
+	IAssetEditorInstance*        AssetEditor = FAssetEditorManager::Get().FindEditorForAsset(InSequence, false);
+	ILevelSequenceEditorToolkit* LevelSequenceEditor = static_cast<ILevelSequenceEditorToolkit*>(AssetEditor);
+
+	TSharedPtr<ISequencer> SequencerPtr = LevelSequenceEditor ? LevelSequenceEditor->GetSequencer() : nullptr;
+
+	if (!SequencerPtr.IsValid())
+	{
+		return;
+	}
+
+	UTakeRecorderSources* Sources = InSequence->FindOrAddMetaData<UTakeRecorderSources>();
+
+	TArray<AActor*> ActorsBeingRecorded;
+	if (Sources)
+	{
+		for (UTakeRecorderSource* Source : Sources->GetSources())
+		{
+			if (Source->IsA<UTakeRecorderActorSource>())
+			{
+				UTakeRecorderActorSource* ActorSource = Cast<UTakeRecorderActorSource>(Source);
+				if (ActorSource->Target.IsValid())
+				{
+					ActorsBeingRecorded.Add(ActorSource->Target.Get());
+				}
+			}
+		}
+	}
+
+	TArray<TWeakObjectPtr<AActor>> ActorsToAdd;
+	for (ULevel* Level : InWorld->GetLevels())
+	{
+		if (Level)
+		{
+			for (AActor* Actor : Level->Actors)
+			{
+				if (!ActorsBeingRecorded.Contains(Actor))
+				{
+					ActorsToAdd.Add(Actor);
+				}
+			}
+		}
+	}
+
+	TArray<FGuid> AddedGuids = SequencerPtr->AddActors(ActorsToAdd, false);
+
+	// Add to the "Autotracked Changes" folder
+	UMovieScene* MovieScene = InSequence->GetMovieScene();
+	FName AutotrackedFolderName = FName(TEXT("Autotracked Changes"));
+
+	UMovieSceneFolder* FolderToUse = nullptr;
+	for (UMovieSceneFolder* Folder : MovieScene->GetRootFolders())
+	{
+		if (Folder->GetFolderName() == AutotrackedFolderName)
+		{
+			FolderToUse = Folder;
+			break;
+		}
+	}
+
+	if (FolderToUse == nullptr)
+	{
+		FolderToUse = NewObject<UMovieSceneFolder>(MovieScene, NAME_None, RF_Transactional);
+		FolderToUse->SetFolderName(AutotrackedFolderName);
+		MovieScene->GetRootFolders().Add(FolderToUse);
+	}
+
+	for (FGuid AddedGuid : AddedGuids)
+	{
+		FolderToUse->AddChildObjectBinding(AddedGuid);
+	}
+}
\ No newline at end of file
diff --git a/Engine/Plugins/VirtualProduction/Takes/Source/TakeRecorderSources/Private/TakeRecorderWorldSource.h b/Engine/Plugins/VirtualProduction/Takes/Source/TakeRecorderSources/Private/TakeRecorderWorldSource.h
new file mode 100644
--- /dev/null
+++ b/Engine/Plugins/VirtualProduction/Takes/Source/TakeRecorderSources/Private/TakeRecorderWorldSource.h
@@ -0,0 +1,65 @@
+// Copyright 1998-2019 Epic Games, Inc. All Rights Reserved.
+
+#pragma once
+
+#include "TakeRecorderSource.h"
+#include "TakeRecorderSources.h"
+#include "UObject/SoftObjectPtr.h"
+#include "Templates/SubclassOf.h"
+#include "GameFramework/Actor.h"
+#include "TakeRecorderWorldSource.generated.h"
+
+class UTexture;
+class UTakeRecorderActorSource;
+
+/** A recording source that records world state */
+UCLASS(Abstract, config = EditorSettings, DisplayName = "World Recorder Defaults")
+class UTakeRecorderWorldSourceSettings : public UTakeRecorderSource
+{
+public:
+	GENERATED_BODY()
+
+	UTakeRecorderWorldSourceSettings(const FObjectInitializer& ObjInit);
+
+	virtual void PostEditChangeProperty(FPropertyChangedEvent& PropertyChangedEvent) override;
+
+	/** Record world settings */
+	UPROPERTY(config, EditAnywhere, BlueprintReadWrite, Category = "Source")
+	bool bRecordWorldSettings;
+
+	/** Add a binding and track for all actors that aren't explicitly being recorded */
+	UPROPERTY(config, EditAnywhere, BlueprintReadWrite, Category = "Source")
+	bool bAutotrackActors;
+};
+
+
+/** A recording source that records world state */
+UCLASS(DisplayName="World", Category="Actors")
+class UTakeRecorderWorldSource : public UTakeRecorderWorldSourceSettings
+{
+public:
+	GENERATED_BODY()
+
+		UTakeRecorderWorldSource(const FObjectInitializer& ObjInit);
+
+private:
+
+	// UTakeRecorderSource
+	virtual TArray<UTakeRecorderSource*> PreRecording(class ULevelSequence* InSequence, class ULevelSequence* InMasterSequence, FManifestSerializer* InManifestSerializer) override;
+	virtual TArray<UTakeRecorderSource*> PostRecording(class ULevelSequence* InSequence, class ULevelSequence* InMasterSequence) override;
+	virtual bool SupportsTakeNumber() const override { return false; }
+	virtual FText GetDisplayTextImpl() const override;
+	virtual bool CanAddSource(UTakeRecorderSources* InSources) const override;
+
+	// This source does not support subscenes (ie. "World Settings subscene"), but the world settings actor would be placed in subscenes if the option is enabled
+	virtual bool SupportsSubscenes() const override { return false; }
+
+private:
+
+	/*
+	 * Autotrack actors in the world that aren't already being recorded
+	 */
+	void AutotrackActors(class ULevelSequence* InSequence, UWorld* InWorld);
+
+	TWeakObjectPtr<UTakeRecorderActorSource> WorldSource;
+};
diff --git a/Engine/Plugins/VirtualProduction/Takes/Source/TakeRecorderSources/TakeRecorderSources.Build.cs b/Engine/Plugins/VirtualProduction/Takes/Source/TakeRecorderSources/TakeRecorderSources.Build.cs
--- a/Engine/Plugins/VirtualProduction/Takes/Source/TakeRecorderSources/TakeRecorderSources.Build.cs
+++ b/Engine/Plugins/VirtualProduction/Takes/Source/TakeRecorderSources/TakeRecorderSources.Build.cs
@@ -13,7 +13,8 @@ public class TakeRecorderSources : ModuleRules
 				"CoreUObject",
 				"Engine",
 				"LevelSequence",
-				"MovieScene",
+                "LevelSequenceEditor",
+                "MovieScene",
 				"MovieSceneTracks",
                 "SceneOutliner",
 				"SequenceRecorder", // For ISequenceAudioRecorder
diff --git a/Engine/Source/Editor/Sequencer/Private/Sequencer.cpp b/Engine/Source/Editor/Sequencer/Private/Sequencer.cpp
--- a/Engine/Source/Editor/Sequencer/Private/Sequencer.cpp
+++ b/Engine/Source/Editor/Sequencer/Private/Sequencer.cpp
@@ -4823,7 +4823,7 @@ void FSequencer::SaveCurrentMovieSceneAs()
 }
 
 
-TArray<FGuid> FSequencer::AddActors(const TArray<TWeakObjectPtr<AActor> >& InActors)
+TArray<FGuid> FSequencer::AddActors(const TArray<TWeakObjectPtr<AActor> >& InActors, bool bSelectActors)
 {
 	TArray<FGuid> PossessableGuids;
 
@@ -4889,22 +4889,25 @@ TArray<FGuid> FSequencer::AddActors(const TArray<TWeakObjectPtr<AActor> >& InAct
 			}
 		}
 
-		// Clear our editor selection so we can make the selection our added actors.
-		// This has to be done after we know if the actor is going to be added to a
-		// folder, otherwise it causes the folder we wanted to pick to be deselected.
-		USelection* SelectedActors = GEditor->GetSelectedActors();
-		SelectedActors->BeginBatchSelectOperation();
-		SelectedActors->Modify();
-		GEditor->SelectNone(false, true);
-		for (TWeakObjectPtr<AActor> WeakActor : InActors)
+		if (bSelectActors)
 		{
-			if (AActor* Actor = WeakActor.Get())
+			// Clear our editor selection so we can make the selection our added actors.
+			// This has to be done after we know if the actor is going to be added to a
+			// folder, otherwise it causes the folder we wanted to pick to be deselected.
+			USelection* SelectedActors = GEditor->GetSelectedActors();
+			SelectedActors->BeginBatchSelectOperation();
+			SelectedActors->Modify();
+			GEditor->SelectNone(false, true);
+			for (TWeakObjectPtr<AActor> WeakActor : InActors)
 			{
-				GEditor->SelectActor(Actor, true, false);
+				if (AActor* Actor = WeakActor.Get())
+				{
+					GEditor->SelectActor(Actor, true, false);
+				}
 			}
+			SelectedActors->EndBatchSelectOperation();
+			GEditor->NoteSelectionChange();
 		}
-		SelectedActors->EndBatchSelectOperation();
-		GEditor->NoteSelectionChange();
 
 		// Add the possessables as children of the first selected folder
 		if (SelectedParentFolders.Num() > 0)
@@ -4916,12 +4919,15 @@ TArray<FGuid> FSequencer::AddActors(const TArray<TWeakObjectPtr<AActor> >& InAct
 		}
 
 		// Now add them all to the selection set to be selected after a tree rebuild.
-		for(const FGuid& Possessable : PossessableGuids)
+		if (bSelectActors)
 		{
-			FString PossessablePath = NewNodePath += Possessable.ToString();
-			
-			// Object Bindings use their FGuid as their unique key.
-			SequencerWidget->AddAdditionalPathToSelectionSet(PossessablePath);
+			for (const FGuid& Possessable : PossessableGuids)
+			{
+				FString PossessablePath = NewNodePath += Possessable.ToString();
+
+				// Object Bindings use their FGuid as their unique key.
+				SequencerWidget->AddAdditionalPathToSelectionSet(PossessablePath);
+			}
 		}
 
 		RefreshTree();
diff --git a/Engine/Source/Editor/Sequencer/Private/Sequencer.h b/Engine/Source/Editor/Sequencer/Private/Sequencer.h
--- a/Engine/Source/Editor/Sequencer/Private/Sequencer.h
+++ b/Engine/Source/Editor/Sequencer/Private/Sequencer.h
@@ -692,7 +692,7 @@ public:
 	virtual void UpdatePlaybackRange() override;
 	virtual void SetPlaybackSpeed(float InPlaybackSpeed) override { PlaybackSpeed = InPlaybackSpeed; }
 	virtual float GetPlaybackSpeed() const override { return PlaybackSpeed; }
-	virtual TArray<FGuid> AddActors(const TArray<TWeakObjectPtr<AActor> >& InActors) override;
+	virtual TArray<FGuid> AddActors(const TArray<TWeakObjectPtr<AActor> >& InActors, bool bSelectActors = true) override;
 	virtual void AddSubSequence(UMovieSceneSequence* Sequence) override;
 	virtual bool CanKeyProperty(FCanKeyPropertyParams CanKeyPropertyParams) const override;
 	virtual void KeyProperty(FKeyPropertyParams KeyPropertyParams) override;
diff --git a/Engine/Source/Editor/Sequencer/Public/ISequencer.h b/Engine/Source/Editor/Sequencer/Public/ISequencer.h
--- a/Engine/Source/Editor/Sequencer/Public/ISequencer.h
+++ b/Engine/Source/Editor/Sequencer/Public/ISequencer.h
@@ -232,9 +232,10 @@ public:
 	 * Add actors as possessable objects to sequencer.
 	 * 
 	 * @param InActors The actors to add to sequencer.
+	 * @param bSelectActors Select the newly added possessable objects in sequencer.
 	 * @return The posssessable guids for the newly added actors.
 	 */
-	virtual TArray<FGuid> AddActors(const TArray<TWeakObjectPtr<AActor> >& InActors) = 0;
+	virtual TArray<FGuid> AddActors(const TArray<TWeakObjectPtr<AActor> >& InActors, bool bSelectActors = true) = 0;
 
 	/**
 	 * Calling this function will add the specified track to the currently selected folder
