Take Recorder: Take recorder display names now pulled from metadata rather than DisplayName.

diff --git a/Engine/Plugins/Animation/LiveLink/Source/LiveLinkEditor/Private/TakeRecorderSource/TakeRecorderLiveLinkSource.h b/Engine/Plugins/Animation/LiveLink/Source/LiveLinkEditor/Private/TakeRecorderSource/TakeRecorderLiveLinkSource.h
--- a/Engine/Plugins/Animation/LiveLink/Source/LiveLinkEditor/Private/TakeRecorderSource/TakeRecorderLiveLinkSource.h
+++ b/Engine/Plugins/Animation/LiveLink/Source/LiveLinkEditor/Private/TakeRecorderSource/TakeRecorderLiveLinkSource.h
@@ -46,7 +46,7 @@ public:
 };
 
 /** A recording source that records LiveLink */
-UCLASS(DisplayName="Live Link", Category="Live Link", config=EditorSettings)
+UCLASS(Category="Live Link", config=EditorSettings, meta = (TakeRecorderDisplayName = "Live Link"))
 class UTakeRecorderLiveLinkSource : public UTakeRecorderSource
 {
 
diff --git a/Engine/Plugins/VirtualProduction/Takes/Source/TakeRecorder/Private/Widgets/SLevelSequenceTakeEditor.cpp b/Engine/Plugins/VirtualProduction/Takes/Source/TakeRecorder/Private/Widgets/SLevelSequenceTakeEditor.cpp
--- a/Engine/Plugins/VirtualProduction/Takes/Source/TakeRecorder/Private/Widgets/SLevelSequenceTakeEditor.cpp
+++ b/Engine/Plugins/VirtualProduction/Takes/Source/TakeRecorder/Private/Widgets/SLevelSequenceTakeEditor.cpp
@@ -259,7 +259,7 @@ TSharedRef<SWidget> SLevelSequenceTakeEditor::OnGenerateSourcesMenu()
 			TSubclassOf<UTakeRecorderSource> SubclassOf = Class;
 
 			MenuBuilder.AddMenuEntry(
-				Class->GetDisplayNameText(),
+				FText::FromString(Class->GetMetaData(TEXT("TakeRecorderDisplayName"))),
 				Class->GetToolTipText(true),
 				FSlateIconFinder::FindIconForClass(Class),
 				FUIAction(
diff --git a/Engine/Plugins/VirtualProduction/Takes/Source/TakeRecorderSources/Private/TakeRecorderActorSource.h b/Engine/Plugins/VirtualProduction/Takes/Source/TakeRecorderSources/Private/TakeRecorderActorSource.h
--- a/Engine/Plugins/VirtualProduction/Takes/Source/TakeRecorderSources/Private/TakeRecorderActorSource.h
+++ b/Engine/Plugins/VirtualProduction/Takes/Source/TakeRecorderSources/Private/TakeRecorderActorSource.h
@@ -31,7 +31,7 @@ enum class ETakeRecorderActorRecordType : uint8
 * Records the properties of the actor and the components on the actor and safely
 * handles new components being spawned at runtime and the actor being destroyed.
 */
-UCLASS(DisplayName="Any Actor", Category="Actors")
+UCLASS(Category="Actors", meta = (TakeRecorderDisplayName="Any Actor"))
 class UTakeRecorderActorSource : public UTakeRecorderSource, public IMovieSceneTrackRecorderHost
 {
 public:
diff --git a/Engine/Plugins/VirtualProduction/Takes/Source/TakeRecorderSources/Private/TakeRecorderLevelSequenceSource.h b/Engine/Plugins/VirtualProduction/Takes/Source/TakeRecorderSources/Private/TakeRecorderLevelSequenceSource.h
--- a/Engine/Plugins/VirtualProduction/Takes/Source/TakeRecorderSources/Private/TakeRecorderLevelSequenceSource.h
+++ b/Engine/Plugins/VirtualProduction/Takes/Source/TakeRecorderSources/Private/TakeRecorderLevelSequenceSource.h
@@ -11,7 +11,7 @@ class ALevelSequenceActor;
 class ULevelSequence;
 
 /** Plays level sequence actors when recording starts */
-UCLASS(DisplayName="Level Sequence", Category="Other")
+UCLASS(Category="Other", meta = (TakeRecorderDisplayName = "Level Sequence"))
 class UTakeRecorderLevelSequenceSource : public UTakeRecorderSource
 {
 public:
diff --git a/Engine/Plugins/VirtualProduction/Takes/Source/TakeRecorderSources/Private/TakeRecorderLevelVisibilitySource.h b/Engine/Plugins/VirtualProduction/Takes/Source/TakeRecorderSources/Private/TakeRecorderLevelVisibilitySource.h
--- a/Engine/Plugins/VirtualProduction/Takes/Source/TakeRecorderSources/Private/TakeRecorderLevelVisibilitySource.h
+++ b/Engine/Plugins/VirtualProduction/Takes/Source/TakeRecorderSources/Private/TakeRecorderLevelVisibilitySource.h
@@ -28,7 +28,7 @@ public:
 };
 
 /** A recording source that records level visibility state */
-UCLASS(DisplayName="Level Visibility", Category="Other", config = EditorSettings)
+UCLASS(Category="Other", config = EditorSettings, meta = (TakeRecorderDisplayName = "Level Visibility"))
 class UTakeRecorderLevelVisibilitySource : public UTakeRecorderLevelVisibilitySourceSettings
 {
 public:
diff --git a/Engine/Plugins/VirtualProduction/Takes/Source/TakeRecorderSources/Private/TakeRecorderMicrophoneAudioSource.h b/Engine/Plugins/VirtualProduction/Takes/Source/TakeRecorderSources/Private/TakeRecorderMicrophoneAudioSource.h
--- a/Engine/Plugins/VirtualProduction/Takes/Source/TakeRecorderSources/Private/TakeRecorderMicrophoneAudioSource.h
+++ b/Engine/Plugins/VirtualProduction/Takes/Source/TakeRecorderSources/Private/TakeRecorderMicrophoneAudioSource.h
@@ -35,7 +35,7 @@ public:
 };
 
 /** A recording source that records microphone audio */
-UCLASS(DisplayName="Microphone Audio", Category="Audio", config=EditorSettings)
+UCLASS(Category="Audio", config=EditorSettings, meta = (TakeRecorderDisplayName = "Microphone Audio"))
 class UTakeRecorderMicrophoneAudioSource : public UTakeRecorderMicrophoneAudioSourceSettings
 {
 public:
diff --git a/Engine/Plugins/VirtualProduction/Takes/Source/TakeRecorderSources/Private/TakeRecorderNearbySpawnedActorSource.h b/Engine/Plugins/VirtualProduction/Takes/Source/TakeRecorderSources/Private/TakeRecorderNearbySpawnedActorSource.h
--- a/Engine/Plugins/VirtualProduction/Takes/Source/TakeRecorderSources/Private/TakeRecorderNearbySpawnedActorSource.h
+++ b/Engine/Plugins/VirtualProduction/Takes/Source/TakeRecorderSources/Private/TakeRecorderNearbySpawnedActorSource.h
@@ -12,7 +12,7 @@ class UTexture;
 class UTakeRecorderActorSource;
 
 /** A recording source that detects actors spawned close to the current camera, and captures them as spawnables */
-UCLASS(DisplayName="Nearby Spawned Actors", Category="Actors")
+UCLASS(Category="Actors", meta = (TakeRecorderDisplayName = "Nearby Spawned Actors"))
 class UTakeRecorderNearbySpawnedActorSource : public UTakeRecorderSource
 {
 public:
diff --git a/Engine/Plugins/VirtualProduction/Takes/Source/TakeRecorderSources/Private/TakeRecorderPlayerSource.h b/Engine/Plugins/VirtualProduction/Takes/Source/TakeRecorderSources/Private/TakeRecorderPlayerSource.h
--- a/Engine/Plugins/VirtualProduction/Takes/Source/TakeRecorderSources/Private/TakeRecorderPlayerSource.h
+++ b/Engine/Plugins/VirtualProduction/Takes/Source/TakeRecorderSources/Private/TakeRecorderPlayerSource.h
@@ -10,7 +10,7 @@
 class UTakeRecorderActorSource;
 
 /** A recording source that records the current player */
-UCLASS(DisplayName="Player", Category="Actors")
+UCLASS(Category="Actors", meta = (TakeRecorderDisplayName = "Player"))
 class UTakeRecorderPlayerSource : public UTakeRecorderSource
 {
 public:
diff --git a/Engine/Plugins/VirtualProduction/Takes/Source/TakeRecorderSources/Private/TakeRecorderWorldSource.h b/Engine/Plugins/VirtualProduction/Takes/Source/TakeRecorderSources/Private/TakeRecorderWorldSource.h
--- a/Engine/Plugins/VirtualProduction/Takes/Source/TakeRecorderSources/Private/TakeRecorderWorldSource.h
+++ b/Engine/Plugins/VirtualProduction/Takes/Source/TakeRecorderSources/Private/TakeRecorderWorldSource.h
@@ -34,7 +34,7 @@ public:
 
 
 /** A recording source that records world state */
-UCLASS(DisplayName="World", Category="Actors")
+UCLASS(Category="Actors", meta = (TakeRecorderDisplayName = "World"))
 class UTakeRecorderWorldSource : public UTakeRecorderWorldSourceSettings
 {
 public:
