TakeRecorder

- Modified AddSourceForActor to only add when the actor isn't already there to avoid having duplicates in the source list.
- Also added RemoveActorFromSources to remove an actor directly from the sources.

diff --git a/Engine/Plugins/VirtualProduction/Takes/Source/TakeRecorderSources/Private/TakeRecorderActorSource.cpp b/Engine/Plugins/VirtualProduction/Takes/Source/TakeRecorderSources/Private/TakeRecorderActorSource.cpp
--- a/Engine/Plugins/VirtualProduction/Takes/Source/TakeRecorderSources/Private/TakeRecorderActorSource.cpp
+++ b/Engine/Plugins/VirtualProduction/Takes/Source/TakeRecorderSources/Private/TakeRecorderActorSource.cpp
@@ -44,7 +44,6 @@ static const FName MovieSceneSectionRecorderFactoryName("MovieSceneTrackRecorder
 static const FName SequencerTrackClassMetadataName("SequencerTrackClass");
 static const FName DoNotRecordTag("DoNotRecord");
 
-
 UTakeRecorderSource* UTakeRecorderActorSource::AddSourceForActor(AActor* InActor, UTakeRecorderSources* InSources)
 {
 	if (InSources == nullptr)
@@ -52,12 +51,122 @@ UTakeRecorderSource* UTakeRecorderActorSource::AddSourceForActor(AActor* InActor
 		FFrame::KismetExecutionMessage(TEXT("The Source is invalid."), ELogVerbosity::Error);
 		return nullptr;
 	}
+	
+	if (InActor == nullptr)
+	{
+		FFrame::KismetExecutionMessage(TEXT("The Actor is invalid."), ELogVerbosity::Error);
+		return nullptr;
+	}
+
+	//Look through our sources and see if one actor matches the incoming one either from editor or PIE world.
+	{
+		//Cache  InputActor comparison data
+		const bool bIsAlreadyPIEActor = InActor->GetOutermost()->HasAnyPackageFlags(PKG_PlayInEditor);
+		const AActor* InputActorEditor = EditorUtilities::GetEditorWorldCounterpartActor(InActor);
+		const AActor* InputActorPIE = EditorUtilities::GetSimWorldCounterpartActor(InActor);
+
+		TArray<UTakeRecorderSource*> SourceArray = InSources->GetSourcesCopy();
+		for (UTakeRecorderSource* CurrentSource : SourceArray)
+		{
+			UTakeRecorderActorSource* CurrentActorSource = Cast<UTakeRecorderActorSource>(CurrentSource);
+			if (CurrentActorSource != nullptr)
+			{
+				AActor* CurrentActor = CurrentActorSource->Target.Get();
+				if (CurrentActor == nullptr)
+				{
+					continue;
+				}
+
+				if (InActor == CurrentActor)
+				{
+					return CurrentActorSource;
+				}
+				else 
+				{
+					if (bIsAlreadyPIEActor)
+					{
+						//The input actor is from PIE -> Bring it into Editor world and compare. 
+						if (InputActorEditor == CurrentActor)
+						{
+							return CurrentActorSource;
+						}
+					}
+					else
+					{
+						//The input actor is from Editor -> Bring it into PIE world and compare. 
+						if (InputActorPIE == CurrentActor)
+						{
+							return CurrentActorSource;
+						}
+					}
+				}
+			}
+		}
+	}
 
 	UTakeRecorderActorSource* NewSource = InSources->AddSource<UTakeRecorderActorSource>();
 	NewSource->SetSourceActor(InActor);
 	return NewSource;
 }
 
+void UTakeRecorderActorSource::RemoveActorFromSources(AActor* InActor, UTakeRecorderSources* InSources)
+{
+	if (InSources == nullptr)
+	{
+		FFrame::KismetExecutionMessage(TEXT("The Source to remove from is invalid."), ELogVerbosity::Error);
+		return;
+	}
+
+	if (InActor == nullptr)
+	{
+		FFrame::KismetExecutionMessage(TEXT("The Actor to remove is invalid."), ELogVerbosity::Error);
+		return;
+	}
+
+	//Cache  InputActor comparison data
+	const bool bIsAlreadyPIEActor = InActor->GetOutermost()->HasAnyPackageFlags(PKG_PlayInEditor);
+	const AActor* InputActorEditor = EditorUtilities::GetEditorWorldCounterpartActor(InActor);
+	const AActor* InputActorPIE = EditorUtilities::GetSimWorldCounterpartActor(InActor);
+
+	TArray<UTakeRecorderSource*> SourceArray = InSources->GetSourcesCopy();
+	for (UTakeRecorderSource* CurrentSource : SourceArray)
+	{
+		UTakeRecorderActorSource* CurrentActorSource = Cast<UTakeRecorderActorSource>(CurrentSource);
+		if (CurrentActorSource != nullptr)
+		{
+			const AActor* CurrentActor = CurrentActorSource->Target.Get();
+			if (CurrentActor == nullptr)
+			{
+				continue;
+			}
+
+			if (InActor == CurrentActor)
+			{
+				InSources->RemoveSource(CurrentSource);
+			}
+			else
+			{
+				if (bIsAlreadyPIEActor)
+				{
+					//The input actor is from PIE -> Bring it into Editor world and compare. 
+					if (InputActorEditor == CurrentActor)
+					{
+						InSources->RemoveSource(CurrentSource);
+					}
+				}
+				else
+				{
+					//The input actor is from Editor -> Bring it into PIE world and compare. 
+					if (InputActorPIE == CurrentActor)
+					{
+						InSources->RemoveSource(CurrentSource);
+					}
+				}
+			}
+		}
+	}
+}
+
 UTakeRecorderActorSource::UTakeRecorderActorSource(const FObjectInitializer& ObjInit)
 	: Super(ObjInit)
 {
diff --git a/Engine/Plugins/VirtualProduction/Takes/Source/TakeRecorderSources/Private/TakeRecorderActorSource.h b/Engine/Plugins/VirtualProduction/Takes/Source/TakeRecorderSources/Private/TakeRecorderActorSource.h
--- a/Engine/Plugins/VirtualProduction/Takes/Source/TakeRecorderSources/Private/TakeRecorderActorSource.h
+++ b/Engine/Plugins/VirtualProduction/Takes/Source/TakeRecorderSources/Private/TakeRecorderActorSource.h
@@ -91,11 +91,22 @@ public:
 	 * Add a take recorder source for the given actor. 
 	 *
 	 * @param InActor The actor to add a source for
-	 * @return The added source
+	 * @param InSources The sources to add the actor to
+	 * @return The added source or the source already present with the same actor
 	 */
 	UFUNCTION(BlueprintCallable, Category = "Take Recorder")
 	static UTakeRecorderSource* AddSourceForActor(AActor* InActor, UTakeRecorderSources* InSources);
 
+	/*
+	 * Remove the given actor from TakeRecorderSources.
+	 *
+	 * @param InActor The actor to remove from the sources
+	 * @param InSources The sources from where to remove the actor
+	 */
+	UFUNCTION(BlueprintCallable, Category = "Take Recorder")
+	static void RemoveActorFromSources(AActor* InActor, UTakeRecorderSources* InSources);
+
+
 public:
 	UTakeRecorderActorSource(const FObjectInitializer& ObjInit);
 	
