Sequence Recorder: Fix asset name refresh

diff --git a/Engine/Plugins/VirtualProduction/Takes/Source/TakeTrackRecorders/Private/MovieSceneAnimationTrackRecorder.cpp b/Engine/Plugins/VirtualProduction/Takes/Source/TakeTrackRecorders/Private/MovieSceneAnimationTrackRecorder.cpp
--- a/Engine/Plugins/VirtualProduction/Takes/Source/TakeTrackRecorders/Private/MovieSceneAnimationTrackRecorder.cpp
+++ b/Engine/Plugins/VirtualProduction/Takes/Source/TakeTrackRecorders/Private/MovieSceneAnimationTrackRecorder.cpp
@@ -8,7 +8,6 @@
 #include "AnimationRecorder.h"
 #include "MovieScene.h"
 #include "AssetRegistryModule.h"
-#include "SequenceRecorderUtils.h"
 #include "SequenceRecorderSettings.h"
 #include "Components/SkeletalMeshComponent.h"
 #include "Animation/AnimSequence.h"
@@ -40,7 +39,7 @@ void UMovieSceneAnimationTrackRecorder::CreateAnimationAssetAndSequence(const AA
 	{
 		ComponentTransform = SkeletalMeshComponent->GetComponentToWorld().GetRelativeTransform(Actor->GetTransform());
 		FString AnimationAssetName = Actor->GetActorLabel();
-		AnimSequence = SequenceRecorderUtils::MakeNewAsset<UAnimSequence>(AnimationDirectory.Path, AnimationAssetName);
+		AnimSequence = TakesUtils::MakeNewAsset<UAnimSequence>(AnimationDirectory.Path, AnimationAssetName);
 		if (AnimSequence.IsValid())
 		{
 			FAssetRegistryModule::AssetCreated(AnimSequence.Get());
diff --git a/Engine/Plugins/VirtualProduction/Takes/Source/TakesCore/Public/TakesUtils.h b/Engine/Plugins/VirtualProduction/Takes/Source/TakesCore/Public/TakesUtils.h
--- a/Engine/Plugins/VirtualProduction/Takes/Source/TakesCore/Public/TakesUtils.h
+++ b/Engine/Plugins/VirtualProduction/Takes/Source/TakesCore/Public/TakesUtils.h
@@ -8,6 +8,7 @@
 #include "IAssetRegistry.h"
 #include "AssetRegistryModule.h"
 #include "UObject/UObjectGlobals.h"
+#include "TakesCoreFwd.h"
 
 class UWorld;
 class ULevelSequence;
@@ -87,4 +88,49 @@ namespace TakesUtils
 
 		return true;
 	}
+
+	/**
+	 * Utility function that creates an asset with the specified asset path and name.
+	 * If the asset cannot be created (as one already exists), we try to postfix the asset
+	 * name until we can successfully create the asset.
+	 */
+	template<typename AssetType>
+	static AssetType* MakeNewAsset(const FString& BaseAssetPath, const FString& BaseAssetName)
+	{
+		const FString Dot(TEXT("."));
+		FString AssetPath = BaseAssetPath;
+		FString AssetName = BaseAssetName;
+		AssetName = AssetName.Replace(TEXT("."), TEXT("_"));
+
+		AssetPath /= AssetName;
+		AssetPath += Dot + AssetName;
+
+		FAssetRegistryModule& AssetRegistryModule = FModuleManager::LoadModuleChecked<FAssetRegistryModule>(TEXT("AssetRegistry"));
+		FAssetData AssetData = AssetRegistryModule.Get().GetAssetByObjectPath(*AssetPath);
+
+		// if object with same name exists, try a different name until we don't find one
+		int32 ExtensionIndex = 0;
+		while (AssetData.IsValid() && AssetData.GetClass() == AssetType::StaticClass())
+		{
+			AssetName = FString::Printf(TEXT("%s_%d"), *AssetName, ExtensionIndex);
+			AssetPath = (BaseAssetPath / AssetName) + Dot + AssetName;
+			AssetData = AssetRegistryModule.Get().GetAssetByObjectPath(*AssetPath);
+
+			ExtensionIndex++;
+		}
+
+		// Create the new asset in the package we just made
+		AssetPath = (BaseAssetPath / AssetName);
+
+		FString FileName;
+		if (FPackageName::TryConvertLongPackageNameToFilename(AssetPath, FileName))
+		{
+			UObject* Package = CreatePackage(nullptr, *AssetPath);
+			return NewObject<AssetType>(Package, *AssetName, RF_Public | RF_Standalone);
+		}
+
+		UE_LOG(LogTakesCore, Error, TEXT("Couldn't create file for package %s"), *AssetPath);
+
+		return nullptr;
+	}
 }
\ No newline at end of file
diff --git a/Engine/Source/Editor/SequenceRecorder/Private/SequenceRecorder.cpp b/Engine/Source/Editor/SequenceRecorder/Private/SequenceRecorder.cpp
--- a/Engine/Source/Editor/SequenceRecorder/Private/SequenceRecorder.cpp
+++ b/Engine/Source/Editor/SequenceRecorder/Private/SequenceRecorder.cpp
@@ -1281,16 +1281,20 @@ void FSequenceRecorder::RefreshNextSequence()
 		SequenceName = GetSequenceRecordingName().Len() > 0 ? GetSequenceRecordingName() : TEXT("RecordedSequence");
 	}
 
+	FString BasePath = GetSequenceRecordingBasePath() / SequenceName;
+
 	// Cache the name of the next sequence we will try to record to
-	NextSequenceName = SequenceRecorderUtils::MakeNewAssetName(GetSequenceRecordingBasePath(), SequenceName);
+	NextSequenceName = SequenceRecorderUtils::MakeNewAssetName(BasePath, SequenceName);
 }
 
 void FSequenceRecorder::ForceRefreshNextSequence()
 {
 	SequenceName = GetSequenceRecordingName().Len() > 0 ? GetSequenceRecordingName() : TEXT("RecordedSequence");
 
+	FString BasePath = GetSequenceRecordingBasePath() / SequenceName;
+
 	// Cache the name of the next sequence we will try to record to
-	NextSequenceName = SequenceRecorderUtils::MakeNewAssetName(GetSequenceRecordingBasePath(), SequenceName);
+	NextSequenceName = SequenceRecorderUtils::MakeNewAssetName(BasePath, SequenceName);
 }
 
 TWeakObjectPtr<ASequenceRecorderGroup> FSequenceRecorder::GetRecordingGroupActor()
diff --git a/Engine/Source/Editor/SequenceRecorder/Public/SequenceRecorderUtils.h b/Engine/Source/Editor/SequenceRecorder/Public/SequenceRecorderUtils.h
--- a/Engine/Source/Editor/SequenceRecorder/Public/SequenceRecorderUtils.h
+++ b/Engine/Source/Editor/SequenceRecorder/Public/SequenceRecorderUtils.h
@@ -30,7 +30,6 @@ static AssetType* MakeNewAsset(const FString& BaseAssetPath, const FString& Base
 	const FString Dot(TEXT("."));
 	FString AssetPath = BaseAssetPath;
 	FString AssetName = BaseAssetName;
-	AssetName = AssetName.Replace(TEXT("."), TEXT("_"));
 
 	AssetPath /= AssetName;
 	AssetPath += Dot + AssetName;
@@ -42,7 +41,7 @@ static AssetType* MakeNewAsset(const FString& BaseAssetPath, const FString& Base
 	int32 ExtensionIndex = 0;
 	while(AssetData.IsValid() && AssetData.GetClass() == AssetType::StaticClass())
 	{
-		AssetName = FString::Printf(TEXT("%s_%d"), *AssetName, ExtensionIndex);
+		AssetName = FString::Printf(TEXT("%s_%d"), *BaseAssetName, ExtensionIndex);
 		AssetPath = (BaseAssetPath / AssetName) + Dot + AssetName;
 		AssetData = AssetRegistryModule.Get().GetAssetByObjectPath(*AssetPath);
 
