Take Recorder: Only set range if max greater than min

diff --git a/Engine/Plugins/VirtualProduction/Takes/Source/TakesCore/Private/TakeRecorderSources.cpp b/Engine/Plugins/VirtualProduction/Takes/Source/TakesCore/Private/TakeRecorderSources.cpp
--- a/Engine/Plugins/VirtualProduction/Takes/Source/TakesCore/Private/TakeRecorderSources.cpp
+++ b/Engine/Plugins/VirtualProduction/Takes/Source/TakesCore/Private/TakeRecorderSources.cpp
@@ -366,8 +366,11 @@ FFrameTime UTakeRecorderSources::TickRecording(class ULevelSequence* InSequence,
 
 		// We're going to use the running time since recording started which is close enough for now until
 		// we get to recording things that get destroyed and needing to stop updating the sub section...
-		TRange<FFrameNumber> CurrentRange = TRange<FFrameNumber>::Exclusive(StartFrame, CurrentFrameTimeSinceStart.FrameNumber);
-		SubSection->SetRange(CurrentRange);
+		if (StartFrame < CurrentFrameTimeSinceStart.FrameNumber)
+		{
+			TRange<FFrameNumber> CurrentRange = TRange<FFrameNumber>::Exclusive(StartFrame, CurrentFrameTimeSinceStart.FrameNumber);
+			SubSection->SetRange(CurrentRange);
+		}
 	}
 	return CurrentFrameTimeSinceStart;
 }
