Take Recorder: Don't record level sequence actors

diff --git a/Engine/Plugins/VirtualProduction/Takes/Source/TakeRecorderSources/Private/TakeRecorderSourcesModule.cpp b/Engine/Plugins/VirtualProduction/Takes/Source/TakeRecorderSources/Private/TakeRecorderSourcesModule.cpp
--- a/Engine/Plugins/VirtualProduction/Takes/Source/TakeRecorderSources/Private/TakeRecorderSourcesModule.cpp
+++ b/Engine/Plugins/VirtualProduction/Takes/Source/TakeRecorderSources/Private/TakeRecorderSourcesModule.cpp
@@ -153,7 +153,7 @@ namespace
 			{
 				for (AActor* Actor : Level->Actors)
 				{
-					if (Actor && Actor->IsA(Class) && !Actor->IsA(ALevelScriptActor::StaticClass()))
+					if (Actor && Actor->IsA(Class) && !Actor->IsA(ALevelScriptActor::StaticClass()) && !Actor->IsA(ALevelSequenceActor::StaticClass()))
 					{
 						OutActors.AddUnique(Actor);
 					}
