Add callback when the panel is constructed.
VirtualCamera adds the camera only if the take panel is constructed.

diff --git a/Engine/Plugins/Media/MediaIOFramework/Source/MediaIOCore/Private/MediaCapture.cpp b/Engine/Plugins/Media/MediaIOFramework/Source/MediaIOCore/Private/MediaCapture.cpp
--- a/Engine/Plugins/Media/MediaIOFramework/Source/MediaIOCore/Private/MediaCapture.cpp
+++ b/Engine/Plugins/Media/MediaIOFramework/Source/MediaIOCore/Private/MediaCapture.cpp
@@ -393,7 +393,7 @@ void UMediaCapture::SetState(EMediaCaptureState InNewState)
 			AsyncTask(ENamedThreads::GameThread, [Self]
 			{
 				UMediaCapture* MediaCapture = Self.Get();
-				if (MediaCapture)
+				if (UObjectInitialized() && MediaCapture)
 				{
 					MediaCapture->BroadcastStateChanged();
 				}
diff --git a/Engine/Plugins/VirtualProduction/Takes/Source/TakeRecorder/Private/Recorder/TakeRecorderBlueprintLibrary.cpp b/Engine/Plugins/VirtualProduction/Takes/Source/TakeRecorder/Private/Recorder/TakeRecorderBlueprintLibrary.cpp
--- a/Engine/Plugins/VirtualProduction/Takes/Source/TakeRecorder/Private/Recorder/TakeRecorderBlueprintLibrary.cpp
+++ b/Engine/Plugins/VirtualProduction/Takes/Source/TakeRecorder/Private/Recorder/TakeRecorderBlueprintLibrary.cpp
@@ -14,6 +14,7 @@
 namespace
 {
 	static UTakeRecorderPanel* CurrentTakeRecorderPanel;
+	static UTakeRecorderBlueprintLibrary::FOnTakeRecorderPanelChanged TakeRecorderPanelChanged;
 }
 
 
@@ -120,20 +121,30 @@ UTakeRecorderPanel* UTakeRecorderBlueprintLibrary::GetTakeRecorderPanel()
 	return CurrentTakeRecorderPanel && CurrentTakeRecorderPanel->IsPanelOpen() ? CurrentTakeRecorderPanel : nullptr;
 }
 
-void UTakeRecorderBlueprintLibrary::SetTakeRecorderPanel(UTakeRecorderPanel* InNewPanel)
+void UTakeRecorderBlueprintLibrary::SetOnTakeRecorderPanelChanged(FOnTakeRecorderPanelChanged OnTakeRecorderPanelChanged)
 {
-	if (CurrentTakeRecorderPanel)
-	{
-		// Old panel is no longer valid
-		CurrentTakeRecorderPanel->ClosePanel();
-
-		CurrentTakeRecorderPanel->RemoveFromRoot();
-		CurrentTakeRecorderPanel = nullptr;
-	}
+	TakeRecorderPanelChanged = OnTakeRecorderPanelChanged;
+};
 
-	if (InNewPanel && InNewPanel->IsPanelOpen())
+void UTakeRecorderBlueprintLibrary::SetTakeRecorderPanel(UTakeRecorderPanel* InNewPanel)
+{
+	if (CurrentTakeRecorderPanel != InNewPanel)
 	{
-		InNewPanel->AddToRoot();
-		CurrentTakeRecorderPanel = InNewPanel;
+		if (CurrentTakeRecorderPanel)
+		{
+			// Old panel is no longer valid
+			CurrentTakeRecorderPanel->ClosePanel();
+
+			CurrentTakeRecorderPanel->RemoveFromRoot();
+			CurrentTakeRecorderPanel = nullptr;
+		}
+
+		if (InNewPanel && InNewPanel->IsPanelOpen())
+		{
+			InNewPanel->AddToRoot();
+			CurrentTakeRecorderPanel = InNewPanel;
+		}
+
+		TakeRecorderPanelChanged.ExecuteIfBound();
 	}
 }
\ No newline at end of file
diff --git a/Engine/Plugins/VirtualProduction/Takes/Source/TakeRecorder/Public/Recorder/TakeRecorderBlueprintLibrary.h b/Engine/Plugins/VirtualProduction/Takes/Source/TakeRecorder/Public/Recorder/TakeRecorderBlueprintLibrary.h
--- a/Engine/Plugins/VirtualProduction/Takes/Source/TakeRecorder/Public/Recorder/TakeRecorderBlueprintLibrary.h
+++ b/Engine/Plugins/VirtualProduction/Takes/Source/TakeRecorder/Public/Recorder/TakeRecorderBlueprintLibrary.h
@@ -80,6 +80,13 @@ public:
 	static UTakeRecorderPanel* OpenTakeRecorderPanel();
 
 
+	DECLARE_DYNAMIC_DELEGATE(FOnTakeRecorderPanelChanged);
+
+	/** Called when a Take Panel is constructed or destroyed. */
+	UFUNCTION(BlueprintCallable, Category="Take Recorder")
+	static void SetOnTakeRecorderPanelChanged(FOnTakeRecorderPanelChanged OnTakeRecorderPanelChanged);
+
+
 	/**
 	 * Internal function to assign a new take recorder panel singleton.
 	 * NOTE: Only to be called by STakeRecorderTabContent::Construct.
