Take Recorder: Check for long paths

diff --git a/Engine/Plugins/VirtualProduction/Takes/Source/TakeRecorder/Private/Widgets/STakeRecorderCockpit.cpp b/Engine/Plugins/VirtualProduction/Takes/Source/TakeRecorder/Private/Widgets/STakeRecorderCockpit.cpp
--- a/Engine/Plugins/VirtualProduction/Takes/Source/TakeRecorder/Private/Widgets/STakeRecorderCockpit.cpp
+++ b/Engine/Plugins/VirtualProduction/Takes/Source/TakeRecorder/Private/Widgets/STakeRecorderCockpit.cpp
@@ -465,6 +465,14 @@ void STakeRecorderCockpit::UpdateRecordError()
 		RecordErrorText = FText::Format(LOCTEXT("ErrorWidget_InvalidPath", "{0} is not a valid asset path."), FText::FromString(PackageName));
 		return;
 	}
+
+	const int32 MaxLength = 260;
+
+	if (PackageName.Len() > MaxLength)
+	{
+		RecordErrorText = FText::Format(LOCTEXT("ErrorWidget_TooLongth", "The path to the asset is too long '{0}', the maximum is '{1}'\nPlease choose a shorter name for the slate or create it in a shallower folder structure with shorter folder names."), FText::AsNumber(PackageName.Len()), FText::AsNumber(MaxLength));
+		return;
+	}
 }
 
 void STakeRecorderCockpit::UpdateTakeError()
@@ -939,4 +947,4 @@ TSharedRef<SWidget> STakeRecorderCockpit::MakeLockButton()
 	];
 }
 
-#undef LOCTEXT_NAMESPACE
\ No newline at end of file
+#undef LOCTEXT_NAMESPACE
