Typo

diff --git a/Engine/Plugins/VirtualProduction/Takes/Source/TakeRecorder/Private/Widgets/STakeRecorderCockpit.cpp b/Engine/Plugins/VirtualProduction/Takes/Source/TakeRecorder/Private/Widgets/STakeRecorderCockpit.cpp
--- a/Engine/Plugins/VirtualProduction/Takes/Source/TakeRecorder/Private/Widgets/STakeRecorderCockpit.cpp
+++ b/Engine/Plugins/VirtualProduction/Takes/Source/TakeRecorder/Private/Widgets/STakeRecorderCockpit.cpp
@@ -470,7 +470,7 @@ void STakeRecorderCockpit::UpdateRecordError()
 
 	if (PackageName.Len() > MaxLength)
 	{
-		RecordErrorText = FText::Format(LOCTEXT("ErrorWidget_TooLongth", "The path to the asset is too long '{0}', the maximum is '{1}'\nPlease choose a shorter name for the slate or create it in a shallower folder structure with shorter folder names."), FText::AsNumber(PackageName.Len()), FText::AsNumber(MaxLength));
+		RecordErrorText = FText::Format(LOCTEXT("ErrorWidget_TooLong", "The path to the asset is too long '{0}', the maximum is '{1}'\nPlease choose a shorter name for the slate or create it in a shallower folder structure with shorter folder names."), FText::AsNumber(PackageName.Len()), FText::AsNumber(MaxLength));
 		return;
 	}
 }
