Take Recorder: Add BP delegate for take recorder started, finished, cancelled, slate changed, take number changed

diff --git a/Engine/Plugins/MovieScene/SequencerScripting/Source/SequencerScripting/Private/ExtensionLibraries/MovieSceneSequenceExtensions.cpp b/Engine/Plugins/MovieScene/SequencerScripting/Source/SequencerScripting/Private/ExtensionLibraries/MovieSceneSequenceExtensions.cpp
--- a/Engine/Plugins/MovieScene/SequencerScripting/Source/SequencerScripting/Private/ExtensionLibraries/MovieSceneSequenceExtensions.cpp
+++ b/Engine/Plugins/MovieScene/SequencerScripting/Source/SequencerScripting/Private/ExtensionLibraries/MovieSceneSequenceExtensions.cpp
@@ -514,13 +514,14 @@ TArray<FMovieSceneMarkedFrame> UMovieSceneSequenceExtensions::GetMarkedFrames(UM
 	return TArray<FMovieSceneMarkedFrame>();
 }
 
-void UMovieSceneSequenceExtensions::AddMarkedFrame(UMovieSceneSequence* Sequence, const FMovieSceneMarkedFrame& InMarkedFrame)
+int32 UMovieSceneSequenceExtensions::AddMarkedFrame(UMovieSceneSequence* Sequence, const FMovieSceneMarkedFrame& InMarkedFrame)
 {
 	UMovieScene* MovieScene = Sequence->GetMovieScene();
 	if (MovieScene)
 	{
-		MovieScene->AddMarkedFrame(InMarkedFrame);
+		return MovieScene->AddMarkedFrame(InMarkedFrame);
 	}
+	return INDEX_NONE;
 }
 
 void UMovieSceneSequenceExtensions::RemoveMarkedFrame(UMovieSceneSequence* Sequence, int32 RemoveIndex)
diff --git a/Engine/Plugins/MovieScene/SequencerScripting/Source/SequencerScripting/Public/ExtensionLibraries/MovieSceneSequenceExtensions.h b/Engine/Plugins/MovieScene/SequencerScripting/Source/SequencerScripting/Public/ExtensionLibraries/MovieSceneSequenceExtensions.h
--- a/Engine/Plugins/MovieScene/SequencerScripting/Source/SequencerScripting/Public/ExtensionLibraries/MovieSceneSequenceExtensions.h
+++ b/Engine/Plugins/MovieScene/SequencerScripting/Source/SequencerScripting/Public/ExtensionLibraries/MovieSceneSequenceExtensions.h
@@ -405,9 +405,10 @@ public:
 	 * A unique label will be generated if the marked frame label is empty
 	 *
 	 * @InMarkedFrame The given user marked frame to add
+	 * @return The index to the newly added marked frame
 	 */
 	UFUNCTION(BlueprintCallable, Category = "Sequence", meta = (ScriptMethod))
-	static void AddMarkedFrame(UMovieSceneSequence* Sequence, const FMovieSceneMarkedFrame& InMarkedFrame);
+	static int32 AddMarkedFrame(UMovieSceneSequence* Sequence, const FMovieSceneMarkedFrame& InMarkedFrame);
 
 	/*
 	 * Remove the user marked frame by index.
diff --git a/Engine/Plugins/VirtualProduction/Takes/Source/TakeRecorder/Private/Recorder/TakeRecorder.cpp b/Engine/Plugins/VirtualProduction/Takes/Source/TakeRecorder/Private/Recorder/TakeRecorder.cpp
--- a/Engine/Plugins/VirtualProduction/Takes/Source/TakeRecorder/Private/Recorder/TakeRecorder.cpp
+++ b/Engine/Plugins/VirtualProduction/Takes/Source/TakeRecorder/Private/Recorder/TakeRecorder.cpp
@@ -1,6 +1,7 @@
 // Copyright 1998-2019 Epic Games, Inc. All Rights Reserved.
 
 #include "Recorder/TakeRecorder.h"
+#include "Recorder/TakeRecorderBlueprintLibrary.h"
 #include "TakePreset.h"
 #include "TakeMetaData.h"
 #include "TakeRecorderSources.h"
@@ -311,7 +312,7 @@ public:
 
 	virtual UWorld* GetTickableGameObjectWorld() const override
 	{
-		UTakeRecorder* Recorder = WeakRecorder.Get();
+			UTakeRecorder* Recorder = WeakRecorder.Get();
 		return Recorder ? Recorder->GetWorld() : nullptr;
 	}
 
@@ -686,6 +687,8 @@ void UTakeRecorder::Start()
 		Sequencer->RefreshTree();
 	}
 	OnRecordingStartedEvent.Broadcast(this);
+
+	UTakeRecorderBlueprintLibrary::OnTakeRecorderStarted();
 }
 
 void UTakeRecorder::Stop()
@@ -778,10 +781,12 @@ void UTakeRecorder::Stop()
 		if (bDidEverStartRecording)
 		{
 			OnRecordingFinishedEvent.Broadcast(this);
+			UTakeRecorderBlueprintLibrary::OnTakeRecorderFinished(SequenceAsset);
 		}
 		else
 		{
 			OnRecordingCancelledEvent.Broadcast(this);
+			UTakeRecorderBlueprintLibrary::OnTakeRecorderCancelled();
 		}
 	}
 
diff --git a/Engine/Plugins/VirtualProduction/Takes/Source/TakeRecorder/Private/Recorder/TakeRecorderBlueprintLibrary.cpp b/Engine/Plugins/VirtualProduction/Takes/Source/TakeRecorder/Private/Recorder/TakeRecorderBlueprintLibrary.cpp
--- a/Engine/Plugins/VirtualProduction/Takes/Source/TakeRecorder/Private/Recorder/TakeRecorderBlueprintLibrary.cpp
+++ b/Engine/Plugins/VirtualProduction/Takes/Source/TakeRecorder/Private/Recorder/TakeRecorderBlueprintLibrary.cpp
@@ -15,6 +15,12 @@ namespace
 {
 	static UTakeRecorderPanel* CurrentTakeRecorderPanel;
 	static UTakeRecorderBlueprintLibrary::FOnTakeRecorderPanelChanged TakeRecorderPanelChanged;
+	static UTakeRecorderBlueprintLibrary::FOnTakeRecorderStarted TakeRecorderStarted;
+	static UTakeRecorderBlueprintLibrary::FOnTakeRecorderFinished TakeRecorderFinished;
+	static UTakeRecorderBlueprintLibrary::FOnTakeRecorderCancelled TakeRecorderCancelled;
+	static UTakeRecorderBlueprintLibrary::FOnTakeRecorderMarkedFrameAdded TakeRecorderMarkedFrameAdded;
+	static UTakeRecorderBlueprintLibrary::FOnTakeRecorderSlateChanged TakeRecorderSlateChanged;
+	static UTakeRecorderBlueprintLibrary::FOnTakeRecorderTakeNumberChanged TakeRecorderTakeNumberChanged;
 }
 
 
@@ -126,6 +132,66 @@ void UTakeRecorderBlueprintLibrary::SetOnTakeRecorderPanelChanged(FOnTakeRecorde
 	TakeRecorderPanelChanged = OnTakeRecorderPanelChanged;
 };
 
+void UTakeRecorderBlueprintLibrary::SetOnTakeRecorderStarted(FOnTakeRecorderStarted OnTakeRecorderStarted)
+{
+	TakeRecorderStarted = OnTakeRecorderStarted;
+};
+
+void UTakeRecorderBlueprintLibrary::SetOnTakeRecorderFinished(FOnTakeRecorderFinished OnTakeRecorderFinished)
+{
+	TakeRecorderFinished = OnTakeRecorderFinished;
+};
+
+void UTakeRecorderBlueprintLibrary::SetOnTakeRecorderCancelled(FOnTakeRecorderCancelled OnTakeRecorderCancelled)
+{
+	TakeRecorderCancelled = OnTakeRecorderCancelled;
+};
+
+void UTakeRecorderBlueprintLibrary::SetOnTakeRecorderMarkedFrameAdded(FOnTakeRecorderMarkedFrameAdded OnTakeRecorderMarkedFrameAdded)
+{
+	TakeRecorderMarkedFrameAdded = OnTakeRecorderMarkedFrameAdded;
+};
+
+void UTakeRecorderBlueprintLibrary::SetOnTakeRecorderSlateChanged(FOnTakeRecorderSlateChanged OnTakeRecorderSlateChanged)
+{
+	TakeRecorderSlateChanged = OnTakeRecorderSlateChanged;
+};
+
+void UTakeRecorderBlueprintLibrary::SetOnTakeRecorderTakeNumberChanged(FOnTakeRecorderTakeNumberChanged OnTakeRecorderTakeNumberChanged)
+{
+	TakeRecorderTakeNumberChanged = OnTakeRecorderTakeNumberChanged;
+};
+
+void UTakeRecorderBlueprintLibrary::OnTakeRecorderStarted()
+{
+	TakeRecorderStarted.ExecuteIfBound();
+}
+
+void UTakeRecorderBlueprintLibrary::OnTakeRecorderFinished(ULevelSequence* InSequenceAsset)
+{
+	TakeRecorderFinished.ExecuteIfBound(InSequenceAsset);
+}
+
+void UTakeRecorderBlueprintLibrary::OnTakeRecorderCancelled()
+{
+	TakeRecorderCancelled.ExecuteIfBound();
+}
+
+void UTakeRecorderBlueprintLibrary::OnTakeRecorderMarkedFrameAdded(const FMovieSceneMarkedFrame& InMarkedFrame)
+{
+	TakeRecorderMarkedFrameAdded.ExecuteIfBound(InMarkedFrame);
+}
+
+void UTakeRecorderBlueprintLibrary::OnTakeRecorderSlateChanged(const FString& InSlate)
+{
+	TakeRecorderSlateChanged.ExecuteIfBound(InSlate);
+}
+
+void UTakeRecorderBlueprintLibrary::OnTakeRecorderTakeNumberChanged(int32 InTakeNumber)
+{
+	TakeRecorderTakeNumberChanged.ExecuteIfBound(InTakeNumber);
+}
+
 void UTakeRecorderBlueprintLibrary::SetTakeRecorderPanel(UTakeRecorderPanel* InNewPanel)
 {
 	if (CurrentTakeRecorderPanel != InNewPanel)
diff --git a/Engine/Plugins/VirtualProduction/Takes/Source/TakeRecorder/Private/Widgets/STakeRecorderCockpit.cpp b/Engine/Plugins/VirtualProduction/Takes/Source/TakeRecorder/Private/Widgets/STakeRecorderCockpit.cpp
--- a/Engine/Plugins/VirtualProduction/Takes/Source/TakeRecorder/Private/Widgets/STakeRecorderCockpit.cpp
+++ b/Engine/Plugins/VirtualProduction/Takes/Source/TakeRecorder/Private/Widgets/STakeRecorderCockpit.cpp
@@ -100,6 +100,7 @@ void STakeRecorderCockpit::Construct(const FArguments& InArgs)
 		if (NextTakeNumber != TakeMetaData->GetTakeNumber())
 		{
 			TakeMetaData->SetTakeNumber(NextTakeNumber);
+			UTakeRecorderBlueprintLibrary::OnTakeRecorderTakeNumberChanged(NextTakeNumber);
 		}
 	}
 
@@ -546,11 +547,20 @@ void STakeRecorderCockpit::CacheMetaData()
 			TransientTakeMetaData = UTakeMetaData::CreateFromDefaults(GetTransientPackage(), NAME_None);
 			TransientTakeMetaData->SetFlags(RF_Transactional | RF_Transient);
 
-			TransientTakeMetaData->SetSlate(GetDefault<UTakeRecorderProjectSettings>()->Settings.DefaultSlate);
+			FString DefaultSlate = GetDefault<UTakeRecorderProjectSettings>()->Settings.DefaultSlate;
+			if (TransientTakeMetaData->GetSlate() != DefaultSlate)
+			{
+				TransientTakeMetaData->SetSlate(DefaultSlate);
+				UTakeRecorderBlueprintLibrary::OnTakeRecorderSlateChanged(DefaultSlate);
+			}
 
 			// Compute the correct starting take number
 			int32 NextTakeNumber = UTakesCoreBlueprintLibrary::ComputeNextTakeNumber(TransientTakeMetaData->GetSlate());
-			TransientTakeMetaData->SetTakeNumber(NextTakeNumber);
+			if (TransientTakeMetaData->GetTakeNumber() != NextTakeNumber)
+			{
+				TransientTakeMetaData->SetTakeNumber(NextTakeNumber);
+				UTakeRecorderBlueprintLibrary::OnTakeRecorderTakeNumberChanged(NextTakeNumber);
+			}
 		}
 
 		NewMetaDataThisTick = TransientTakeMetaData;
@@ -633,10 +643,16 @@ void STakeRecorderCockpit::SetSlateText(const FText& InNewText, ETextCommit::Typ
 		TakeMetaData->Modify();
 
 		TakeMetaData->SetSlate(InNewText.ToString());
+		UTakeRecorderBlueprintLibrary::OnTakeRecorderSlateChanged(InNewText.ToString());
 
 		// Compute the correct starting take number
 		int32 NextTakeNumber = UTakesCoreBlueprintLibrary::ComputeNextTakeNumber(TakeMetaData->GetSlate());
-		TakeMetaData->SetTakeNumber(NextTakeNumber);
+
+		if (NextTakeNumber != TakeMetaData->GetTakeNumber())
+		{
+			TakeMetaData->SetTakeNumber(NextTakeNumber);
+			UTakeRecorderBlueprintLibrary::OnTakeRecorderTakeNumberChanged(NextTakeNumber);
+		}
 	}
 }
 
@@ -686,12 +702,16 @@ int32 STakeRecorderCockpit::GetTakeNumber() const
 
 FReply STakeRecorderCockpit::OnSetNextTakeNumber()
 {
-	FScopedTransaction Transaction(LOCTEXT("SetNextTakeNumber_Transaction", "Set Next Take Number"));
-
 	int32 NextTakeNumber = UTakesCoreBlueprintLibrary::ComputeNextTakeNumber(TakeMetaData->GetSlate());
 
-	TakeMetaData->Modify();
-	TakeMetaData->SetTakeNumber(NextTakeNumber);
+	if (TakeMetaData->GetTakeNumber() != NextTakeNumber)
+	{
+		FScopedTransaction Transaction(LOCTEXT("SetNextTakeNumber_Transaction", "Set Next Take Number"));
+
+		TakeMetaData->Modify();
+		TakeMetaData->SetTakeNumber(NextTakeNumber);
+		UTakeRecorderBlueprintLibrary::OnTakeRecorderTakeNumberChanged(NextTakeNumber);
+	}
 
 	return FReply::Handled();
 }
@@ -732,9 +752,10 @@ void STakeRecorderCockpit::SetTakeNumber_FromCommit(int32 InNewTakeNumber, EText
 			OnEndSetTakeNumber(InNewTakeNumber);
 		}
 	}
-	else
+	else if (TakeMetaData->GetTakeNumber() != InNewTakeNumber)
 	{
 		TakeMetaData->SetTakeNumber(InNewTakeNumber);
+		UTakeRecorderBlueprintLibrary::OnTakeRecorderTakeNumberChanged(InNewTakeNumber);
 	}
 
 	bAutoApplyTakeNumber = false;
@@ -750,6 +771,7 @@ void STakeRecorderCockpit::OnEndSetTakeNumber(int32 InFinalValue)
 	}
 
 	TakeMetaData->SetTakeNumber(InFinalValue);
+	UTakeRecorderBlueprintLibrary::OnTakeRecorderTakeNumberChanged(InFinalValue);
 
 	GEditor->EndTransaction();
 	TransactionIndex = INDEX_NONE;
@@ -771,7 +793,8 @@ FReply STakeRecorderCockpit::OnAddMarkedFrame()
 		FMovieSceneMarkedFrame MarkedFrame;
 		MarkedFrame.FrameNumber = ConvertFrameTime(ElapsedFrame, MovieScene->GetDisplayRate(), MovieScene->GetTickResolution()).CeilToFrame();
 
-		MovieScene->AddMarkedFrame(MarkedFrame);
+		int32 MarkedFrameIndex = MovieScene->AddMarkedFrame(MarkedFrame);
+		UTakeRecorderBlueprintLibrary::OnTakeRecorderMarkedFrameAdded(MovieScene->GetMarkedFrames()[MarkedFrameIndex]);
 	}
 
 	return FReply::Handled();
@@ -897,7 +920,12 @@ void STakeRecorderCockpit::OnRecordingFinished(UTakeRecorder* Recorder)
 	{
 		// Increment the transient take meta data if necessary
 		int32 NextTakeNumber = UTakesCoreBlueprintLibrary::ComputeNextTakeNumber(TransientTakeMetaData->GetSlate());
-		TransientTakeMetaData->SetTakeNumber(NextTakeNumber);
+
+		if (TransientTakeMetaData->GetTakeNumber() != NextTakeNumber)
+		{
+			TransientTakeMetaData->SetTakeNumber(NextTakeNumber);
+			UTakeRecorderBlueprintLibrary::OnTakeRecorderTakeNumberChanged(NextTakeNumber);
+		}
 
 		bAutoApplyTakeNumber = true;
 	}
diff --git a/Engine/Plugins/VirtualProduction/Takes/Source/TakeRecorder/Public/Recorder/TakeRecorderBlueprintLibrary.h b/Engine/Plugins/VirtualProduction/Takes/Source/TakeRecorder/Public/Recorder/TakeRecorderBlueprintLibrary.h
--- a/Engine/Plugins/VirtualProduction/Takes/Source/TakeRecorder/Public/Recorder/TakeRecorderBlueprintLibrary.h
+++ b/Engine/Plugins/VirtualProduction/Takes/Source/TakeRecorder/Public/Recorder/TakeRecorderBlueprintLibrary.h
@@ -81,11 +81,47 @@ public:
 
 
 	DECLARE_DYNAMIC_DELEGATE(FOnTakeRecorderPanelChanged);
+	DECLARE_DYNAMIC_DELEGATE(FOnTakeRecorderStarted);
+	DECLARE_DYNAMIC_DELEGATE_OneParam(FOnTakeRecorderFinished, ULevelSequence*, SequenceAsset);
+	DECLARE_DYNAMIC_DELEGATE(FOnTakeRecorderCancelled);
+	DECLARE_DYNAMIC_DELEGATE_OneParam(FOnTakeRecorderMarkedFrameAdded, const FMovieSceneMarkedFrame&, MarkedFrame);
+	DECLARE_DYNAMIC_DELEGATE_OneParam(FOnTakeRecorderSlateChanged, const FString&, Slate);
+	DECLARE_DYNAMIC_DELEGATE_OneParam(FOnTakeRecorderTakeNumberChanged, int32, TakeNumber);
 
 	/** Called when a Take Panel is constructed or destroyed. */
 	UFUNCTION(BlueprintCallable, Category="Take Recorder")
 	static void SetOnTakeRecorderPanelChanged(FOnTakeRecorderPanelChanged OnTakeRecorderPanelChanged);
 
+	/** Called when take recording starts. */
+	UFUNCTION(BlueprintCallable, Category = "Take Recorder")
+	static void SetOnTakeRecorderStarted(FOnTakeRecorderStarted OnTakeRecorderStarted);
+
+	/** Called when take recording finishes. */
+	UFUNCTION(BlueprintCallable, Category = "Take Recorder")
+	static void SetOnTakeRecorderFinished(FOnTakeRecorderFinished OnTakeRecorderFinished);
+
+	/** Called when take recording is cancelled. */
+	UFUNCTION(BlueprintCallable, Category = "Take Recorder")
+	static void SetOnTakeRecorderCancelled(FOnTakeRecorderCancelled OnTakeRecorderCancelled);
+
+	/** Called when a marked frame is added. */
+	UFUNCTION(BlueprintCallable, Category = "Take Recorder")
+	static void SetOnTakeRecorderMarkedFrameAdded(FOnTakeRecorderMarkedFrameAdded OnTakeRecorderMarkedFrameAdded);
+
+	/** Called when the slate is changed. */
+	UFUNCTION(BlueprintCallable, Category = "Take Recorder")
+	static void SetOnTakeRecorderSlateChanged(FOnTakeRecorderSlateChanged OnTakeRecorderSlateChanged);
+
+	/** Called when the take number is changed. */
+	UFUNCTION(BlueprintCallable, Category = "Take Recorder")
+	static void SetOnTakeRecorderTakeNumberChanged(FOnTakeRecorderTakeNumberChanged OnTakeRecorderTakeNumberChanged);
+
+	static void OnTakeRecorderStarted();
+	static void OnTakeRecorderFinished(ULevelSequence* InSequenceAsset);
+	static void OnTakeRecorderCancelled();
+	static void OnTakeRecorderMarkedFrameAdded(const FMovieSceneMarkedFrame& InMarkedFrame);
+	static void OnTakeRecorderSlateChanged(const FString& InSlate);
+	static void OnTakeRecorderTakeNumberChanged(int32 InTakeNumber);
 
 	/**
 	 * Internal function to assign a new take recorder panel singleton.
