Take Recorder: Split GetSubsceneName to GetSubsceneTrackName and GetSubsceneAssetName.

- Append slate and take number to actor source sub level sequences.
- Change logic to find the matching track name rather than the asset name when replacing subsections.
- Swap order of names for audio, level vis, and live link track names to make it more consistent with actor sources.

ie. If recording Scene_1 and Take 10 for 'Cube', the sub level sequence will be called: 'Cube_Scene_1_0010' rather than "Cube"
ie. If recording Scene_1 audio, the sub level sequence will be called: "Audio_Scene_1" rather than "Scene_1Audio"

diff --git a/Engine/Plugins/Animation/LiveLink/Source/LiveLinkEditor/Private/TakeRecorderSource/TakeRecorderLiveLinkSource.cpp b/Engine/Plugins/Animation/LiveLink/Source/LiveLinkEditor/Private/TakeRecorderSource/TakeRecorderLiveLinkSource.cpp
--- a/Engine/Plugins/Animation/LiveLink/Source/LiveLinkEditor/Private/TakeRecorderSource/TakeRecorderLiveLinkSource.cpp
+++ b/Engine/Plugins/Animation/LiveLink/Source/LiveLinkEditor/Private/TakeRecorderSource/TakeRecorderLiveLinkSource.cpp
@@ -97,11 +97,25 @@ bool UTakeRecorderLiveLinkSource::CanAddSource(UTakeRecorderSources* InSources)
 	return true;
 }
 
-FString UTakeRecorderLiveLinkSource::GetSubsceneName(ULevelSequence* InSequence) const
+FString UTakeRecorderLiveLinkSource::GetSubsceneTrackName(ULevelSequence* InSequence) const
 {
 	if (UTakeMetaData* TakeMetaData = InSequence->FindMetaData<UTakeMetaData>())
 	{
-		return TakeMetaData->GetSlate() + SubjectName.ToString();
+		return FString::Printf(TEXT("%s_%s"), *SubjectName.ToString(), *TakeMetaData->GenerateAssetPath("{slate}"));
+	}
+	else if (SubjectName != NAME_None)
+	{
+		return SubjectName.ToString();
+	}
+
+	return TEXT("LiveLink");
+}
+
+FString UTakeRecorderLiveLinkSource::GetSubsceneAssetName(ULevelSequence* InSequence) const
+{
+	if (UTakeMetaData* TakeMetaData = InSequence->FindMetaData<UTakeMetaData>())
+	{
+		return FString::Printf(TEXT("%s_%s"), *SubjectName.ToString(), *TakeMetaData->GenerateAssetPath("{slate}_{take}"));
 	}
 	else if (SubjectName != NAME_None)
 	{
diff --git a/Engine/Plugins/Animation/LiveLink/Source/LiveLinkEditor/Private/TakeRecorderSource/TakeRecorderLiveLinkSource.h b/Engine/Plugins/Animation/LiveLink/Source/LiveLinkEditor/Private/TakeRecorderSource/TakeRecorderLiveLinkSource.h
--- a/Engine/Plugins/Animation/LiveLink/Source/LiveLinkEditor/Private/TakeRecorderSource/TakeRecorderLiveLinkSource.h
+++ b/Engine/Plugins/Animation/LiveLink/Source/LiveLinkEditor/Private/TakeRecorderSource/TakeRecorderLiveLinkSource.h
@@ -81,7 +81,8 @@ private:
 	virtual void AddContentsToFolder(class UMovieSceneFolder* InFolder) override;
 	virtual bool CanAddSource(UTakeRecorderSources* InSources) const override;
 	virtual bool SupportsSubscenes() const override { return true; }
-	virtual FString GetSubsceneName(ULevelSequence* InSequence) const override;
+	virtual FString GetSubsceneTrackName(ULevelSequence* InSequence) const override;
+	virtual FString GetSubsceneAssetName(ULevelSequence* InSequence) const override;
 	// ~UTakeRecorderSource
 
 };
diff --git a/Engine/Plugins/VirtualProduction/Takes/Source/TakeRecorder/Private/SerializedRecorder.cpp b/Engine/Plugins/VirtualProduction/Takes/Source/TakeRecorder/Private/SerializedRecorder.cpp
--- a/Engine/Plugins/VirtualProduction/Takes/Source/TakeRecorder/Private/SerializedRecorder.cpp
+++ b/Engine/Plugins/VirtualProduction/Takes/Source/TakeRecorder/Private/SerializedRecorder.cpp
@@ -236,7 +236,7 @@ bool FSerializedRecorder::LoadSubSequenceFile(UMovieSceneSequence* InMovieSceneS
 
 			ULevelSequence* TargetSequence = InMasterSequence;
 			const FString& SubSequenceName = Header.Name;
-			TargetSequence = UTakeRecorderSources::CreateSubSequenceForSource(InMasterSequence, SubSequenceName);
+			TargetSequence = UTakeRecorderSources::CreateSubSequenceForSource(InMasterSequence, SubSequenceName, SubSequenceName);
 			if (TargetSequence)
 			{
 				TargetSequence->GetMovieScene()->TimecodeSource = FApp::GetTimecode();
diff --git a/Engine/Plugins/VirtualProduction/Takes/Source/TakeRecorderSources/Private/TakeRecorderActorSource.cpp b/Engine/Plugins/VirtualProduction/Takes/Source/TakeRecorderSources/Private/TakeRecorderActorSource.cpp
--- a/Engine/Plugins/VirtualProduction/Takes/Source/TakeRecorderSources/Private/TakeRecorderActorSource.cpp
+++ b/Engine/Plugins/VirtualProduction/Takes/Source/TakeRecorderSources/Private/TakeRecorderActorSource.cpp
@@ -14,6 +14,7 @@
 #include "Recorder/TakeRecorderParameters.h"
 #include "TakeRecorderSettings.h"
 #include "TakesUtils.h"
+#include "TakeMetaData.h"
 #include "MovieSceneFolder.h"
 #include "Serializers/MovieSceneManifestSerialization.h"
 
@@ -1574,14 +1575,28 @@ UMovieSceneTrackRecorderSettings* UTakeRecorderActorSource::GetSettingsObjectFor
 	return nullptr;
 }
 
-FString UTakeRecorderActorSource::GetSubsceneName(ULevelSequence* InSequence) const
+FString UTakeRecorderActorSource::GetSubsceneTrackName(ULevelSequence* InSequence) const
 {
 	if (Target.IsValid())
 	{
 		return Target->GetActorLabel();
 	}
 
-	return Super::GetSubsceneName(InSequence);
+	return Super::GetSubsceneTrackName(InSequence);
+}
+
+FString UTakeRecorderActorSource::GetSubsceneAssetName(ULevelSequence* InSequence) const
+{
+	if (Target.IsValid())
+	{
+		if (UTakeMetaData* TakeMetaData = InSequence->FindMetaData<UTakeMetaData>())
+		{
+			return FString::Printf(TEXT("%s_%s"), *Target->GetActorLabel(), *TakeMetaData->GenerateAssetPath("{slate}_{take}"));
+		}
+		return Target->GetActorLabel();
+	}
+
+	return Super::GetSubsceneAssetName(InSequence);
 }
 
 void UTakeRecorderActorSource::AddContentsToFolder(UMovieSceneFolder* InFolder)
diff --git a/Engine/Plugins/VirtualProduction/Takes/Source/TakeRecorderSources/Private/TakeRecorderActorSource.h b/Engine/Plugins/VirtualProduction/Takes/Source/TakeRecorderSources/Private/TakeRecorderActorSource.h
--- a/Engine/Plugins/VirtualProduction/Takes/Source/TakeRecorderSources/Private/TakeRecorderActorSource.h
+++ b/Engine/Plugins/VirtualProduction/Takes/Source/TakeRecorderSources/Private/TakeRecorderActorSource.h
@@ -117,7 +117,8 @@ public:
 	virtual void StopRecording(class ULevelSequence* InSequence) override;
 	virtual TArray<UTakeRecorderSource*> PostRecording(class ULevelSequence* InSequence, class ULevelSequence* InMasterSequence) override;
 	virtual TArray<UObject*> GetAdditionalSettingsObjects() const { return TArray<UObject*>(FactorySettings); }
-	virtual FString GetSubsceneName(ULevelSequence* InSequence) const override;
+	virtual FString GetSubsceneTrackName(ULevelSequence* InSequence) const override;
+	virtual FString GetSubsceneAssetName(ULevelSequence* InSequence) const override;
 	virtual void AddContentsToFolder(class UMovieSceneFolder* InFolder) override;
 	// ~UTakeRecorderSource Interface
 
diff --git a/Engine/Plugins/VirtualProduction/Takes/Source/TakeRecorderSources/Private/TakeRecorderLevelVisibilitySource.cpp b/Engine/Plugins/VirtualProduction/Takes/Source/TakeRecorderSources/Private/TakeRecorderLevelVisibilitySource.cpp
--- a/Engine/Plugins/VirtualProduction/Takes/Source/TakeRecorderSources/Private/TakeRecorderLevelVisibilitySource.cpp
+++ b/Engine/Plugins/VirtualProduction/Takes/Source/TakeRecorderSources/Private/TakeRecorderLevelVisibilitySource.cpp
@@ -35,11 +35,20 @@ void UTakeRecorderLevelVisibilitySourceSettings::PostEditChangeProperty(FPropert
 	}
 }
 
-FString UTakeRecorderLevelVisibilitySourceSettings::GetSubsceneName(ULevelSequence* InSequence) const
+FString UTakeRecorderLevelVisibilitySourceSettings::GetSubsceneTrackName(ULevelSequence* InSequence) const
 {
 	if (UTakeMetaData* TakeMetaData = InSequence->FindMetaData<UTakeMetaData>())
 	{
-		return TakeMetaData->GetSlate() + TEXT("Level Visibility");
+		return FString::Printf(TEXT("Level Visibility_%s"), *TakeMetaData->GenerateAssetPath("{slate}"));
+	}
+	return TEXT("Level Visibility");
+}
+
+FString UTakeRecorderLevelVisibilitySourceSettings::GetSubsceneAssetName(ULevelSequence* InSequence) const
+{
+	if (UTakeMetaData* TakeMetaData = InSequence->FindMetaData<UTakeMetaData>())
+	{
+		return FString::Printf(TEXT("Level Visibility_%s"), *TakeMetaData->GenerateAssetPath("{slate}_{take}"));
 	}
 	return TEXT("Level Visibility");
 }
diff --git a/Engine/Plugins/VirtualProduction/Takes/Source/TakeRecorderSources/Private/TakeRecorderLevelVisibilitySource.h b/Engine/Plugins/VirtualProduction/Takes/Source/TakeRecorderSources/Private/TakeRecorderLevelVisibilitySource.h
--- a/Engine/Plugins/VirtualProduction/Takes/Source/TakeRecorderSources/Private/TakeRecorderLevelVisibilitySource.h
+++ b/Engine/Plugins/VirtualProduction/Takes/Source/TakeRecorderSources/Private/TakeRecorderLevelVisibilitySource.h
@@ -19,7 +19,8 @@ public:
 	virtual void PostEditChangeProperty(FPropertyChangedEvent& PropertyChangedEvent) override;
 
 	// UTakeRecorderSource Interface
-	virtual FString GetSubsceneName(ULevelSequence* InSequence) const override;
+	virtual FString GetSubsceneTrackName(ULevelSequence* InSequence) const override;
+	virtual FString GetSubsceneAssetName(ULevelSequence* InSequence) const override;
 	// ~UTakeRecorderSource Interface
 
 	/** Name of the recorded level visibility track name */
diff --git a/Engine/Plugins/VirtualProduction/Takes/Source/TakeRecorderSources/Private/TakeRecorderMicrophoneAudioSource.cpp b/Engine/Plugins/VirtualProduction/Takes/Source/TakeRecorderSources/Private/TakeRecorderMicrophoneAudioSource.cpp
--- a/Engine/Plugins/VirtualProduction/Takes/Source/TakeRecorderSources/Private/TakeRecorderMicrophoneAudioSource.cpp
+++ b/Engine/Plugins/VirtualProduction/Takes/Source/TakeRecorderSources/Private/TakeRecorderMicrophoneAudioSource.cpp
@@ -39,11 +39,20 @@ void UTakeRecorderMicrophoneAudioSourceSettings::PostEditChangeProperty(FPropert
 	}
 }
 
-FString UTakeRecorderMicrophoneAudioSourceSettings::GetSubsceneName(ULevelSequence* InSequence) const
+FString UTakeRecorderMicrophoneAudioSourceSettings::GetSubsceneTrackName(ULevelSequence* InSequence) const
 {
 	if (UTakeMetaData* TakeMetaData = InSequence->FindMetaData<UTakeMetaData>())
 	{
-		return TakeMetaData->GetSlate() + TEXT("Audio");
+		return FString::Printf(TEXT("Audio_%s"), *TakeMetaData->GenerateAssetPath("{slate}"));
+	}
+	return TEXT("MicrophoneAudio");
+}
+
+FString UTakeRecorderMicrophoneAudioSourceSettings::GetSubsceneAssetName(ULevelSequence* InSequence) const
+{
+	if (UTakeMetaData* TakeMetaData = InSequence->FindMetaData<UTakeMetaData>())
+	{
+		return FString::Printf(TEXT("Audio_%s"), *TakeMetaData->GenerateAssetPath("{slate}_{take}"));
 	}
 	return TEXT("MicrophoneAudio");
 }
diff --git a/Engine/Plugins/VirtualProduction/Takes/Source/TakeRecorderSources/Private/TakeRecorderMicrophoneAudioSource.h b/Engine/Plugins/VirtualProduction/Takes/Source/TakeRecorderSources/Private/TakeRecorderMicrophoneAudioSource.h
--- a/Engine/Plugins/VirtualProduction/Takes/Source/TakeRecorderSources/Private/TakeRecorderMicrophoneAudioSource.h
+++ b/Engine/Plugins/VirtualProduction/Takes/Source/TakeRecorderSources/Private/TakeRecorderMicrophoneAudioSource.h
@@ -22,7 +22,8 @@ public:
 	virtual void PostEditChangeProperty(FPropertyChangedEvent& PropertyChangedEvent) override;
 	
 	// UTakeRecorderSource Interface
-	virtual FString GetSubsceneName(ULevelSequence* InSequence) const override;
+	virtual FString GetSubsceneTrackName(ULevelSequence* InSequence) const override;
+	virtual FString GetSubsceneAssetName(ULevelSequence* InSequence) const override;
 	// ~UTakeRecorderSource Interface
 
 	/** Name of the recorded audio track name */
diff --git a/Engine/Plugins/VirtualProduction/Takes/Source/TakesCore/Private/TakeRecorderSources.cpp b/Engine/Plugins/VirtualProduction/Takes/Source/TakesCore/Private/TakeRecorderSources.cpp
--- a/Engine/Plugins/VirtualProduction/Takes/Source/TakesCore/Private/TakeRecorderSources.cpp
+++ b/Engine/Plugins/VirtualProduction/Takes/Source/TakesCore/Private/TakeRecorderSources.cpp
@@ -98,14 +98,14 @@ void UTakeRecorderSources::StartRecordingRecursive(TArray<UTakeRecorderSource*>
 			// is using which sub-sequence so that we can push the correct sequence for all points of the Source's recording lifecycle.
 			if (bRecordSourcesToSubSequences && Source->SupportsSubscenes())
 			{
-				const FString& SubSequenceName = ObjectTools::SanitizeObjectName(Source->GetSubsceneName(InMasterSequence));
+				const FString& SubSequenceTrackName = ObjectTools::SanitizeObjectName(Source->GetSubsceneTrackName(InMasterSequence));
+				const FString& SubSequenceAssetName = ObjectTools::SanitizeObjectName(Source->GetSubsceneAssetName(InMasterSequence));
 
-				TargetSequence = CreateSubSequenceForSource(InMasterSequence, SubSequenceName);
+				TargetSequence = CreateSubSequenceForSource(InMasterSequence, SubSequenceTrackName, SubSequenceAssetName);
 				TargetSequence->GetMovieScene()->TimecodeSource = Timecode;
 
 				// If there's already a Subscene Track for our sub-sequence we need to remove that track before create a new one. No data is lost in this process as the
 				// sequence that the subscene points to has been copied by CreateSubSequenceForSource so a new track pointed to the new subsequence includes all the old data.
-				TOptional<int32> RowIndex;
 				const FString SequenceName = FPaths::GetBaseFilename(TargetSequence->GetPathName());
 				UMovieSceneSubTrack* SubsceneTrack = nullptr;
 
@@ -113,22 +113,11 @@ void UTakeRecorderSources::StartRecordingRecursive(TArray<UTakeRecorderSource*>
 				{
 					if (Track->IsA<UMovieSceneSubTrack>())
 					{
-						// Look through each section in the track to see if it has a sub-sequence that matches our new sequence.
-						for (UMovieSceneSection* Section : Track->GetAllSections())
+						if (Track->GetDisplayName().ToString() == SubSequenceTrackName)
 						{
-							UMovieSceneSubSection* SubSection = CastChecked<UMovieSceneSubSection>(Section);
-							UMovieSceneSequence* SubSequence = SubSection->GetSequence();
-
-							// Store the row index so we can re-inject the section at the same index to preserve the hierarchical evaluation order.
-							if (FPaths::GetBaseFilename(SubSequence->GetPathName()) == SequenceName)
-							{
-								SubsceneTrack = CastChecked<UMovieSceneSubTrack>(Track);
-								SubsceneTrack->RemoveSection(*Section);
-								RowIndex = Section->GetRowIndex();
-								break;
-							}
+							SubsceneTrack = CastChecked<UMovieSceneSubTrack>(Track);
+							SubsceneTrack->RemoveAllAnimationData();
 						}
-						
 					}
 				}
 
@@ -140,7 +129,7 @@ void UTakeRecorderSources::StartRecordingRecursive(TArray<UTakeRecorderSource*>
 				}
 				
 				// We create a new sub track for every Source so that we can name the Subtrack after the Source instead of just the sections within it.
-				SubsceneTrack->SetDisplayName(FText::FromString(Source->GetSubsceneName(InMasterSequence)));
+				SubsceneTrack->SetDisplayName(FText::FromString(Source->GetSubsceneTrackName(InMasterSequence)));
 				SubsceneTrack->SetColorTint(Source->TrackTint);
 					
 				// When we create the Subscene Track we'll make sure a folder is created for it to sort into and add the new Subscene Track as a child of it.
@@ -150,26 +139,20 @@ void UTakeRecorderSources::StartRecordingRecursive(TArray<UTakeRecorderSource*>
 					Folder->AddChildMasterTrack(SubsceneTrack);
 				}
 				
-				// There isn't already a section for our new sub sequence, so we'll just append it to the end
-				if (!RowIndex.IsSet())
-				{
-					RowIndex = SubsceneTrack->GetMaxRowIndex() + 1;
-				}
-
 				// We initialize the sequence to start at zero and be a 0 frame length section as there is no data in the sections yet.
 				// We'll have to update these sections each frame as the recording progresses so they appear to get longer like normal
 				// tracks do as we record into them.
 				FFrameNumber RecordStartTime = FFrameNumber(0);
 				UMovieSceneSubSection* NewSubSection = SubsceneTrack->AddSequence(TargetSequence, RecordStartTime, 0);
 
-				NewSubSection->SetRowIndex(RowIndex.GetValue());
+				NewSubSection->SetRowIndex(SubsceneTrack->GetMaxRowIndex() + 1);
 				SubsceneTrack->FixRowIndices();
 
 				ActiveSubSections.Add(NewSubSection);
 				if (InManifestSerializer)
 				{
 					FName SerializedType("SubSequence");
-					FManifestProperty  ManifestProperty(SubSequenceName, SerializedType, FGuid());
+					FManifestProperty  ManifestProperty(SubSequenceAssetName, SerializedType, FGuid());
 					InManifestSerializer->WriteFrameData(InManifestSerializer->FramesWritten, ManifestProperty);
 
 					FString AssetPath = InManifestSerializer->GetLocalCaptureDir();
@@ -180,7 +163,7 @@ void UTakeRecorderSources::StartRecordingRecursive(TArray<UTakeRecorderSource*>
 						PlatformFile.CreateDirectory(*AssetPath);
 					}
 
-					AssetPath = AssetPath / SubSequenceName;
+					AssetPath = AssetPath / SubSequenceAssetName;
 					if (!PlatformFile.DirectoryExists(*AssetPath))
 					{
 						PlatformFile.CreateDirectory(*AssetPath);
@@ -192,13 +175,13 @@ void UTakeRecorderSources::StartRecordingRecursive(TArray<UTakeRecorderSource*>
 
 					InManifestSerializer->SetLocalCaptureDir(AssetPath);
 
-					FManifestFileHeader Header(SubSequenceName, SerializedType, FGuid());
+					FManifestFileHeader Header(SubSequenceAssetName, SerializedType, FGuid());
 					FText Error;
-					FString FileName = FString::Printf(TEXT("%s_%s"), *(SerializedType.ToString()), *(SubSequenceName));
+					FString FileName = FString::Printf(TEXT("%s_%s"), *(SerializedType.ToString()), *(SubSequenceAssetName));
 
 					if (!InManifestSerializer->OpenForWrite(FileName, Header, Error))
 					{
-						UE_LOG(SubSequenceSerialization, Warning, TEXT("Error Opening Sequence Sequencer File: Subject '%s' Error '%s'"), *(SubSequenceName), *(Error.ToString()));
+						UE_LOG(SubSequenceSerialization, Warning, TEXT("Error Opening Sequence Sequencer File: Subject '%s' Error '%s'"), *(SubSequenceAssetName), *(Error.ToString()));
 					}
 				}
 			}
@@ -507,7 +490,7 @@ void UTakeRecorderSources::StopRecording(class ULevelSequence* InSequence, FTake
 	CachedLevelSequence = nullptr;
 }
 
-ULevelSequence* UTakeRecorderSources::CreateSubSequenceForSource(ULevelSequence* InMasterSequence, const FString& SubSequenceName)
+ULevelSequence* UTakeRecorderSources::CreateSubSequenceForSource(ULevelSequence* InMasterSequence, const FString& SubSequenceTrackName, const FString& SubSequenceAssetName)
 {
 	IAssetRegistry& AssetRegistry = FModuleManager::LoadModuleChecked<FAssetRegistryModule>("AssetRegistry").Get();
 
@@ -527,16 +510,16 @@ ULevelSequence* UTakeRecorderSources::CreateSubSequenceForSource(ULevelSequence*
 		for (UMovieSceneSection* Section : SubTrack->GetAllSections())
 		{
 			UMovieSceneSubSection* SubSection = CastChecked<UMovieSceneSubSection>(Section);
-			if (FPaths::GetBaseFilename(SubSection->GetSequence()->GetPathName()) == SubSequenceName)
+			if (FPaths::GetBaseFilename(SubSection->GetSequence()->GetPathName()) == SubSequenceAssetName)
 			{
-				UE_LOG(LogTakesCore, Log, TEXT("Found existing sub-section for source %s, duplicating sub-section for recording into."), *SubSequenceName);
+				UE_LOG(LogTakesCore, Log, TEXT("Found existing sub-section for source %s, duplicating sub-section for recording into."), *SubSequenceAssetName);
 				ExistingSubSequence = CastChecked<ULevelSequence>(SubSection->GetSequence());
 				break;
 			}
 		}
 	}
 
-	FString NewPath = FString::Printf(TEXT("%s/%s_Subscenes/%s"), *SequenceDirectory, *SequenceName, *SubSequenceName);
+	FString NewPath = FString::Printf(TEXT("%s/%s_Subscenes/%s"), *SequenceDirectory, *SequenceName, *SubSequenceAssetName);
 	
 	ULevelSequence* OutAsset = nullptr;
 	TakesUtils::CreateNewAssetPackage<ULevelSequence>(NewPath, OutAsset, nullptr, ExistingSubSequence);
@@ -559,7 +542,7 @@ ULevelSequence* UTakeRecorderSources::CreateSubSequenceForSource(ULevelSequence*
 			UTakeMetaData* OutTakeMetaData = OutAsset->CopyMetaData(TakeMetaData);
 
 			// Tack on the sub sequence name so that it's unique from the master sequence
-			OutTakeMetaData->SetSlate(TakeMetaData->GetSlate() + TEXT("_") + SubSequenceName);
+			OutTakeMetaData->SetSlate(TakeMetaData->GetSlate() + TEXT("_") + SubSequenceTrackName);
 		}
 
 		OutAsset->MarkPackageDirty();
diff --git a/Engine/Plugins/VirtualProduction/Takes/Source/TakesCore/Public/TakeRecorderSource.h b/Engine/Plugins/VirtualProduction/Takes/Source/TakesCore/Public/TakeRecorderSource.h
--- a/Engine/Plugins/VirtualProduction/Takes/Source/TakesCore/Public/TakeRecorderSource.h
+++ b/Engine/Plugins/VirtualProduction/Takes/Source/TakesCore/Public/TakeRecorderSource.h
@@ -129,9 +129,14 @@ public:
 	virtual bool SupportsSubscenes() const { return true; }
 
 	/**
-	* When recorded to a Subscenes track, what should the name of the Section be?
+	* When recorded to a Subscenes track, what should the subscene track name of the Section be?
 	*/
-	virtual FString GetSubsceneName(ULevelSequence* InSequence) const { return TEXT("Unnamed_Source"); }
+	virtual FString GetSubsceneTrackName(ULevelSequence* InSequence) const { return TEXT("Unnamed_Source"); }
+
+	/**
+	* When recorded to a Subscenes track, what should the sub level sequence asset name be? By default, returns the name of the track
+	*/
+	virtual FString GetSubsceneAssetName(ULevelSequence* InSequence) const { return GetSubsceneTrackName(InSequence); }
 
 	/**
 	* If you are not recording into a sub-sequence then this will be called after PreRecording is called and will specify the folder
