Take Recorder: Follow up to 5814929

diff --git a/Engine/Plugins/VirtualProduction/Takes/Source/TakesCore/Public/TakeRecorderSources.h b/Engine/Plugins/VirtualProduction/Takes/Source/TakesCore/Public/TakeRecorderSources.h
--- a/Engine/Plugins/VirtualProduction/Takes/Source/TakesCore/Public/TakeRecorderSources.h
+++ b/Engine/Plugins/VirtualProduction/Takes/Source/TakesCore/Public/TakeRecorderSources.h
@@ -147,7 +147,7 @@ public:
 	*/
 
 	/** Creates a sub-sequence asset for the specified sub sequence name based on the given master sequence. */
-	static ULevelSequence* CreateSubSequenceForSource(ULevelSequence* InMasterSequence, const FString& SubSequenceName);
+	static ULevelSequence* CreateSubSequenceForSource(ULevelSequence* InMasterSequence, const FString& SubSequenceTrackName, const FString& SubSequenceAssetName);
 
 private:
 	/** Called at the end of each frame in both the Editor and in Game to update all Sources. */
