Take Recorder: Add project settings to take recorder settings UI

diff --git a/Engine/Plugins/VirtualProduction/Takes/Source/TakeRecorder/Private/Widgets/STakeRecorderPanel.cpp b/Engine/Plugins/VirtualProduction/Takes/Source/TakeRecorder/Private/Widgets/STakeRecorderPanel.cpp
--- a/Engine/Plugins/VirtualProduction/Takes/Source/TakeRecorder/Private/Widgets/STakeRecorderPanel.cpp
+++ b/Engine/Plugins/VirtualProduction/Takes/Source/TakeRecorder/Private/Widgets/STakeRecorderPanel.cpp
@@ -145,10 +145,12 @@ void STakeRecorderPanel::Construct(const FArguments& InArgs)
 		CockpitWidget->GetMetaData()->SetPresetOrigin(InArgs._BasePreset);
 	}
 
-	// Add the user settings immediately if the user preference tells us to
+	// Add the settings immediately if the user preference tells us to
 	UTakeRecorderUserSettings* UserSettings = GetMutableDefault<UTakeRecorderUserSettings>();
+	UTakeRecorderProjectSettings* ProjectSettings = GetMutableDefault<UTakeRecorderProjectSettings>();
 	if (UserSettings->bShowUserSettingsOnUI)
 	{
+		LevelSequenceTakeWidget->AddExternalSettingsObject(ProjectSettings);
 		LevelSequenceTakeWidget->AddExternalSettingsObject(UserSettings);
 	}
 
@@ -406,14 +408,13 @@ TSharedRef<SWidget> STakeRecorderPanel::MakeToolBar()
 			.WidthOverride(ButtonBoxSize)
 			.HeightOverride(ButtonBoxSize)
 			[
-
 				SNew(SCheckBox)
 				.Padding(TakeRecorder::ButtonPadding)
-				.ToolTipText(LOCTEXT("ShowUserSettings_Tip", "Show/Hide the general user settings for take recorder"))
+				.ToolTipText(LOCTEXT("ShowSettings_Tip", "Show/Hide the general user/project settings for take recorder"))
 				.Style(FEditorStyle::Get(), "ToggleButtonCheckbox")
 				.ForegroundColor(FSlateColor::UseForeground())
-				.IsChecked(this, &STakeRecorderPanel::GetUserSettingsCheckState)
-				.OnCheckStateChanged(this, &STakeRecorderPanel::ToggleUserSettings)
+				.IsChecked(this, &STakeRecorderPanel::GetSettingsCheckState)
+				.OnCheckStateChanged(this, &STakeRecorderPanel::ToggleSettings)
 				.Visibility_Lambda([this]() { return !CockpitWidget->Reviewing() ? EVisibility::Visible : EVisibility::Collapsed; })
 				[
 					SNew(STextBlock)
@@ -760,22 +761,25 @@ void STakeRecorderPanel::RefreshPanel()
 }
 
 
-ECheckBoxState STakeRecorderPanel::GetUserSettingsCheckState() const
+ECheckBoxState STakeRecorderPanel::GetSettingsCheckState() const
 {
 	return GetDefault<UTakeRecorderUserSettings>()->bShowUserSettingsOnUI ? ECheckBoxState::Checked : ECheckBoxState::Unchecked;
 }
 
 
-void STakeRecorderPanel::ToggleUserSettings(ECheckBoxState CheckState)
+void STakeRecorderPanel::ToggleSettings(ECheckBoxState CheckState)
 {
 	UTakeRecorderUserSettings* UserSettings = GetMutableDefault<UTakeRecorderUserSettings>();
+	UTakeRecorderProjectSettings* ProjectSettings = GetMutableDefault<UTakeRecorderProjectSettings>();
 
 	if (LevelSequenceTakeWidget->RemoveExternalSettingsObject(UserSettings))
 	{
+		LevelSequenceTakeWidget->RemoveExternalSettingsObject(ProjectSettings);
 		UserSettings->bShowUserSettingsOnUI = false;
 	}
 	else
 	{
+		LevelSequenceTakeWidget->AddExternalSettingsObject(ProjectSettings);
 		LevelSequenceTakeWidget->AddExternalSettingsObject(UserSettings);
 		UserSettings->bShowUserSettingsOnUI = true;
 	}
diff --git a/Engine/Plugins/VirtualProduction/Takes/Source/TakeRecorder/Private/Widgets/STakeRecorderPanel.h b/Engine/Plugins/VirtualProduction/Takes/Source/TakeRecorder/Private/Widgets/STakeRecorderPanel.h
--- a/Engine/Plugins/VirtualProduction/Takes/Source/TakeRecorder/Private/Widgets/STakeRecorderPanel.h
+++ b/Engine/Plugins/VirtualProduction/Takes/Source/TakeRecorder/Private/Widgets/STakeRecorderPanel.h
@@ -100,8 +100,8 @@ private:
 
 	FReply OnReviewLastRecording();
 
-	ECheckBoxState GetUserSettingsCheckState() const;
-	void ToggleUserSettings(ECheckBoxState CheckState);
+	ECheckBoxState GetSettingsCheckState() const;
+	void ToggleSettings(ECheckBoxState CheckState);
 
 	void OnLevelSequenceChanged();
 
diff --git a/Engine/Plugins/VirtualProduction/Takes/Source/TakeRecorderSources/Private/TakeRecorderMicrophoneAudioSource.h b/Engine/Plugins/VirtualProduction/Takes/Source/TakeRecorderSources/Private/TakeRecorderMicrophoneAudioSource.h
--- a/Engine/Plugins/VirtualProduction/Takes/Source/TakeRecorderSources/Private/TakeRecorderMicrophoneAudioSource.h
+++ b/Engine/Plugins/VirtualProduction/Takes/Source/TakeRecorderSources/Private/TakeRecorderMicrophoneAudioSource.h
@@ -11,7 +11,7 @@
 #include "TakeRecorderMicrophoneAudioSource.generated.h"
 
 /** A recording source that records microphone audio */
-UCLASS(Abstract, config=EditorSettings, DisplayName="Microphone Audio Recorder Defaults")
+UCLASS(Abstract, config=EditorSettings, DisplayName="Microphone Audio Recorder")
 class UTakeRecorderMicrophoneAudioSourceSettings : public UTakeRecorderSource
 {
 public:
diff --git a/Engine/Plugins/VirtualProduction/Takes/Source/TakeRecorderSources/Private/TakeRecorderWorldSource.h b/Engine/Plugins/VirtualProduction/Takes/Source/TakeRecorderSources/Private/TakeRecorderWorldSource.h
--- a/Engine/Plugins/VirtualProduction/Takes/Source/TakeRecorderSources/Private/TakeRecorderWorldSource.h
+++ b/Engine/Plugins/VirtualProduction/Takes/Source/TakeRecorderSources/Private/TakeRecorderWorldSource.h
@@ -13,7 +13,7 @@ class UTexture;
 class UTakeRecorderActorSource;
 
 /** A recording source that records world state */
-UCLASS(Abstract, config = EditorSettings, DisplayName = "World Recorder Defaults")
+UCLASS(Abstract, config = EditorSettings, DisplayName = "World Recorder")
 class UTakeRecorderWorldSourceSettings : public UTakeRecorderSource
 {
 public:
diff --git a/Engine/Plugins/VirtualProduction/Takes/Source/TakeTrackRecorders/Public/TrackRecorders/MovieSceneAnimationTrackRecorderSettings.h b/Engine/Plugins/VirtualProduction/Takes/Source/TakeTrackRecorders/Public/TrackRecorders/MovieSceneAnimationTrackRecorderSettings.h
--- a/Engine/Plugins/VirtualProduction/Takes/Source/TakeTrackRecorders/Public/TrackRecorders/MovieSceneAnimationTrackRecorderSettings.h
+++ b/Engine/Plugins/VirtualProduction/Takes/Source/TakeTrackRecorders/Public/TrackRecorders/MovieSceneAnimationTrackRecorderSettings.h
@@ -8,7 +8,7 @@
 #include "Curves/RichCurve.h"
 #include "MovieSceneAnimationTrackRecorderSettings.generated.h"
 
-UCLASS(Abstract, BlueprintType, config=EditorSettings, DisplayName="Animation Recorder Defaults")
+UCLASS(Abstract, BlueprintType, config=EditorSettings, DisplayName="Animation Recorder")
 class TAKETRACKRECORDERS_API UMovieSceneAnimationTrackRecorderEditorSettings : public UMovieSceneTrackRecorderSettings
 {
 	GENERATED_BODY()
