Take Recorder: Add a non-throttled button for the marks button.

diff --git a/Engine/Plugins/VirtualProduction/Takes/Source/TakeRecorder/Private/Widgets/STakeRecorderCockpit.cpp b/Engine/Plugins/VirtualProduction/Takes/Source/TakeRecorder/Private/Widgets/STakeRecorderCockpit.cpp
--- a/Engine/Plugins/VirtualProduction/Takes/Source/TakeRecorder/Private/Widgets/STakeRecorderCockpit.cpp
+++ b/Engine/Plugins/VirtualProduction/Takes/Source/TakeRecorder/Private/Widgets/STakeRecorderCockpit.cpp
@@ -77,7 +77,19 @@ void STakeRecorderCockpit::AddReferencedObjects(FReferenceCollector& Collector)
 	Collector.AddReferencedObject(TransientTakeMetaData);
 }
 
-PRAGMA_DISABLE_OPTIMIZATION
+struct SNonThrottledButton : SButton
+{
+	virtual FReply OnMouseButtonDown(const FGeometry& MyGeometry, const FPointerEvent& MouseEvent) override
+	{
+		FReply Reply = SButton::OnMouseButtonDown(MyGeometry, MouseEvent);
+		if (Reply.IsEventHandled())
+		{
+			Reply.PreventThrottling();
+		}
+		return Reply;
+	}
+};
+
 void STakeRecorderCockpit::Construct(const FArguments& InArgs)
 {
 	FAssetRegistryModule& AssetRegistryModule = FModuleManager::LoadModuleChecked<FAssetRegistryModule>(TEXT("AssetRegistry"));
@@ -339,7 +351,7 @@ void STakeRecorderCockpit::Construct(const FArguments& InArgs)
 					+SHorizontalBox::Slot()
 					.AutoWidth()
 					[
-						SNew(SButton)
+						SNew(SNonThrottledButton)
 						.ButtonStyle(FEditorStyle::Get(), "HoverHintOnly")
 						.ToolTipText(LOCTEXT("AddMarkedFrame", "Click to add a marked frame while recording"))
 						.IsEnabled_Lambda([this]() { return IsRecording() == ECheckBoxState::Checked; })
@@ -391,7 +403,6 @@ void STakeRecorderCockpit::Construct(const FArguments& InArgs)
 		]
 	];
 }
-PRAGMA_ENABLE_OPTIMIZATION
 
 bool STakeRecorderCockpit::CanStartRecording(FText* OutErrorText) const
 {
