Take Recorder: Add bEmitChanged to SetSlate and SetTakeNumber so that firing events can be controlled.

diff --git a/Engine/Plugins/VirtualProduction/Takes/Source/TakeRecorder/Private/Recorder/TakeRecorderBlueprintLibrary.cpp b/Engine/Plugins/VirtualProduction/Takes/Source/TakeRecorder/Private/Recorder/TakeRecorderBlueprintLibrary.cpp
--- a/Engine/Plugins/VirtualProduction/Takes/Source/TakeRecorder/Private/Recorder/TakeRecorderBlueprintLibrary.cpp
+++ b/Engine/Plugins/VirtualProduction/Takes/Source/TakeRecorder/Private/Recorder/TakeRecorderBlueprintLibrary.cpp
@@ -19,8 +19,6 @@ namespace
 	static UTakeRecorderBlueprintLibrary::FOnTakeRecorderFinished TakeRecorderFinished;
 	static UTakeRecorderBlueprintLibrary::FOnTakeRecorderCancelled TakeRecorderCancelled;
 	static UTakeRecorderBlueprintLibrary::FOnTakeRecorderMarkedFrameAdded TakeRecorderMarkedFrameAdded;
-	static UTakeRecorderBlueprintLibrary::FOnTakeRecorderSlateChanged TakeRecorderSlateChanged;
-	static UTakeRecorderBlueprintLibrary::FOnTakeRecorderTakeNumberChanged TakeRecorderTakeNumberChanged;
 }
 
 
@@ -152,16 +150,6 @@ void UTakeRecorderBlueprintLibrary::SetOnTakeRecorderMarkedFrameAdded(FOnTakeRec
 	TakeRecorderMarkedFrameAdded = OnTakeRecorderMarkedFrameAdded;
 };
 
-void UTakeRecorderBlueprintLibrary::SetOnTakeRecorderSlateChanged(FOnTakeRecorderSlateChanged OnTakeRecorderSlateChanged)
-{
-	TakeRecorderSlateChanged = OnTakeRecorderSlateChanged;
-};
-
-void UTakeRecorderBlueprintLibrary::SetOnTakeRecorderTakeNumberChanged(FOnTakeRecorderTakeNumberChanged OnTakeRecorderTakeNumberChanged)
-{
-	TakeRecorderTakeNumberChanged = OnTakeRecorderTakeNumberChanged;
-};
-
 void UTakeRecorderBlueprintLibrary::OnTakeRecorderStarted()
 {
 	TakeRecorderStarted.ExecuteIfBound();
@@ -182,16 +170,6 @@ void UTakeRecorderBlueprintLibrary::OnTakeRecorderMarkedFrameAdded(const FMovieS
 	TakeRecorderMarkedFrameAdded.ExecuteIfBound(InMarkedFrame);
 }
 
-void UTakeRecorderBlueprintLibrary::OnTakeRecorderSlateChanged(const FString& InSlate)
-{
-	TakeRecorderSlateChanged.ExecuteIfBound(InSlate);
-}
-
-void UTakeRecorderBlueprintLibrary::OnTakeRecorderTakeNumberChanged(int32 InTakeNumber)
-{
-	TakeRecorderTakeNumberChanged.ExecuteIfBound(InTakeNumber);
-}
-
 void UTakeRecorderBlueprintLibrary::SetTakeRecorderPanel(UTakeRecorderPanel* InNewPanel)
 {
 	if (CurrentTakeRecorderPanel != InNewPanel)
diff --git a/Engine/Plugins/VirtualProduction/Takes/Source/TakeRecorder/Private/Widgets/STakeRecorderCockpit.cpp b/Engine/Plugins/VirtualProduction/Takes/Source/TakeRecorder/Private/Widgets/STakeRecorderCockpit.cpp
--- a/Engine/Plugins/VirtualProduction/Takes/Source/TakeRecorder/Private/Widgets/STakeRecorderCockpit.cpp
+++ b/Engine/Plugins/VirtualProduction/Takes/Source/TakeRecorder/Private/Widgets/STakeRecorderCockpit.cpp
@@ -112,7 +112,6 @@ void STakeRecorderCockpit::Construct(const FArguments& InArgs)
 		if (NextTakeNumber != TakeMetaData->GetTakeNumber())
 		{
 			TakeMetaData->SetTakeNumber(NextTakeNumber);
-			UTakeRecorderBlueprintLibrary::OnTakeRecorderTakeNumberChanged(NextTakeNumber);
 		}
 	}
 
@@ -562,7 +561,6 @@ void STakeRecorderCockpit::CacheMetaData()
 			if (TransientTakeMetaData->GetSlate() != DefaultSlate)
 			{
 				TransientTakeMetaData->SetSlate(DefaultSlate);
-				UTakeRecorderBlueprintLibrary::OnTakeRecorderSlateChanged(DefaultSlate);
 			}
 
 			// Compute the correct starting take number
@@ -570,7 +568,6 @@ void STakeRecorderCockpit::CacheMetaData()
 			if (TransientTakeMetaData->GetTakeNumber() != NextTakeNumber)
 			{
 				TransientTakeMetaData->SetTakeNumber(NextTakeNumber);
-				UTakeRecorderBlueprintLibrary::OnTakeRecorderTakeNumberChanged(NextTakeNumber);
 			}
 		}
 
@@ -654,7 +651,6 @@ void STakeRecorderCockpit::SetSlateText(const FText& InNewText, ETextCommit::Typ
 		TakeMetaData->Modify();
 
 		TakeMetaData->SetSlate(InNewText.ToString());
-		UTakeRecorderBlueprintLibrary::OnTakeRecorderSlateChanged(InNewText.ToString());
 
 		// Compute the correct starting take number
 		int32 NextTakeNumber = UTakesCoreBlueprintLibrary::ComputeNextTakeNumber(TakeMetaData->GetSlate());
@@ -662,7 +658,6 @@ void STakeRecorderCockpit::SetSlateText(const FText& InNewText, ETextCommit::Typ
 		if (NextTakeNumber != TakeMetaData->GetTakeNumber())
 		{
 			TakeMetaData->SetTakeNumber(NextTakeNumber);
-			UTakeRecorderBlueprintLibrary::OnTakeRecorderTakeNumberChanged(NextTakeNumber);
 		}
 	}
 }
@@ -721,7 +716,6 @@ FReply STakeRecorderCockpit::OnSetNextTakeNumber()
 
 		TakeMetaData->Modify();
 		TakeMetaData->SetTakeNumber(NextTakeNumber);
-		UTakeRecorderBlueprintLibrary::OnTakeRecorderTakeNumberChanged(NextTakeNumber);
 	}
 
 	return FReply::Handled();
@@ -746,7 +740,7 @@ void STakeRecorderCockpit::SetTakeNumber(int32 InNewTakeNumber)
 
 	if (TransactionIndex != INDEX_NONE || bIsInPIEOrSimulate)
 	{
-		TakeMetaData->SetTakeNumber(InNewTakeNumber);
+		TakeMetaData->SetTakeNumber(InNewTakeNumber, false);
 		bAutoApplyTakeNumber = false;
 	}
 }
@@ -766,7 +760,6 @@ void STakeRecorderCockpit::SetTakeNumber_FromCommit(int32 InNewTakeNumber, EText
 	else if (TakeMetaData->GetTakeNumber() != InNewTakeNumber)
 	{
 		TakeMetaData->SetTakeNumber(InNewTakeNumber);
-		UTakeRecorderBlueprintLibrary::OnTakeRecorderTakeNumberChanged(InNewTakeNumber);
 	}
 
 	bAutoApplyTakeNumber = false;
@@ -782,7 +775,6 @@ void STakeRecorderCockpit::OnEndSetTakeNumber(int32 InFinalValue)
 	}
 
 	TakeMetaData->SetTakeNumber(InFinalValue);
-	UTakeRecorderBlueprintLibrary::OnTakeRecorderTakeNumberChanged(InFinalValue);
 
 	GEditor->EndTransaction();
 	TransactionIndex = INDEX_NONE;
@@ -935,7 +927,6 @@ void STakeRecorderCockpit::OnRecordingFinished(UTakeRecorder* Recorder)
 		if (TransientTakeMetaData->GetTakeNumber() != NextTakeNumber)
 		{
 			TransientTakeMetaData->SetTakeNumber(NextTakeNumber);
-			UTakeRecorderBlueprintLibrary::OnTakeRecorderTakeNumberChanged(NextTakeNumber);
 		}
 
 		bAutoApplyTakeNumber = true;
diff --git a/Engine/Plugins/VirtualProduction/Takes/Source/TakeRecorder/Public/Recorder/TakeRecorderBlueprintLibrary.h b/Engine/Plugins/VirtualProduction/Takes/Source/TakeRecorder/Public/Recorder/TakeRecorderBlueprintLibrary.h
--- a/Engine/Plugins/VirtualProduction/Takes/Source/TakeRecorder/Public/Recorder/TakeRecorderBlueprintLibrary.h
+++ b/Engine/Plugins/VirtualProduction/Takes/Source/TakeRecorder/Public/Recorder/TakeRecorderBlueprintLibrary.h
@@ -85,8 +85,6 @@ public:
 	DECLARE_DYNAMIC_DELEGATE_OneParam(FOnTakeRecorderFinished, ULevelSequence*, SequenceAsset);
 	DECLARE_DYNAMIC_DELEGATE(FOnTakeRecorderCancelled);
 	DECLARE_DYNAMIC_DELEGATE_OneParam(FOnTakeRecorderMarkedFrameAdded, const FMovieSceneMarkedFrame&, MarkedFrame);
-	DECLARE_DYNAMIC_DELEGATE_OneParam(FOnTakeRecorderSlateChanged, const FString&, Slate);
-	DECLARE_DYNAMIC_DELEGATE_OneParam(FOnTakeRecorderTakeNumberChanged, int32, TakeNumber);
 
 	/** Called when a Take Panel is constructed or destroyed. */
 	UFUNCTION(BlueprintCallable, Category="Take Recorder")
@@ -108,20 +106,10 @@ public:
 	UFUNCTION(BlueprintCallable, Category = "Take Recorder")
 	static void SetOnTakeRecorderMarkedFrameAdded(FOnTakeRecorderMarkedFrameAdded OnTakeRecorderMarkedFrameAdded);
 
-	/** Called when the slate is changed. */
-	UFUNCTION(BlueprintCallable, Category = "Take Recorder")
-	static void SetOnTakeRecorderSlateChanged(FOnTakeRecorderSlateChanged OnTakeRecorderSlateChanged);
-
-	/** Called when the take number is changed. */
-	UFUNCTION(BlueprintCallable, Category = "Take Recorder")
-	static void SetOnTakeRecorderTakeNumberChanged(FOnTakeRecorderTakeNumberChanged OnTakeRecorderTakeNumberChanged);
-
 	static void OnTakeRecorderStarted();
 	static void OnTakeRecorderFinished(ULevelSequence* InSequenceAsset);
 	static void OnTakeRecorderCancelled();
 	static void OnTakeRecorderMarkedFrameAdded(const FMovieSceneMarkedFrame& InMarkedFrame);
-	static void OnTakeRecorderSlateChanged(const FString& InSlate);
-	static void OnTakeRecorderTakeNumberChanged(int32 InTakeNumber);
 
 	/**
 	 * Internal function to assign a new take recorder panel singleton.
diff --git a/Engine/Plugins/VirtualProduction/Takes/Source/TakesCore/Private/TakeMetaData.cpp b/Engine/Plugins/VirtualProduction/Takes/Source/TakesCore/Private/TakeMetaData.cpp
--- a/Engine/Plugins/VirtualProduction/Takes/Source/TakesCore/Private/TakeMetaData.cpp
+++ b/Engine/Plugins/VirtualProduction/Takes/Source/TakesCore/Private/TakeMetaData.cpp
@@ -131,19 +131,29 @@ FString UTakeMetaData::GetLevelPath() const
 	return !LevelOrigin.IsNull() ? LevelOrigin.ToString() : FString();
 }
 
-void UTakeMetaData::SetSlate(FString InSlate)
+void UTakeMetaData::SetSlate(FString InSlate, bool bEmitChanged)
 {
 	if (!bIsLocked)
 	{
 		Slate = MoveTemp(InSlate);
+
+		if (bEmitChanged)
+		{
+			UTakesCoreBlueprintLibrary::OnTakeRecorderSlateChanged(Slate);
+		}
 	}
 }
 
-void UTakeMetaData::SetTakeNumber(int32 InTakeNumber)
+void UTakeMetaData::SetTakeNumber(int32 InTakeNumber, bool bEmitChanged)
 {
 	if (!bIsLocked)
 	{
 		TakeNumber = FMath::Max(1, InTakeNumber);
+
+		if (bEmitChanged)
+		{
+			UTakesCoreBlueprintLibrary::OnTakeRecorderTakeNumberChanged(TakeNumber);
+		}
 	}
 }
 
@@ -230,4 +240,5 @@ void UTakeMetaData::ExtendAssetRegistryTagMetaData(TMap<FName, FAssetRegistryTag
 		.SetDisplayName(NSLOCTEXT("TakeMetaData", "LevelPath_Label", "Map"))
 		.SetTooltip(    NSLOCTEXT("TakeMetaData", "LevelPath_Tip",   "Map used for this take"))
 	);
-}
\ No newline at end of file
+}
+
diff --git a/Engine/Plugins/VirtualProduction/Takes/Source/TakesCore/Private/TakeRecorderSources.cpp b/Engine/Plugins/VirtualProduction/Takes/Source/TakesCore/Private/TakeRecorderSources.cpp
--- a/Engine/Plugins/VirtualProduction/Takes/Source/TakesCore/Private/TakeRecorderSources.cpp
+++ b/Engine/Plugins/VirtualProduction/Takes/Source/TakesCore/Private/TakeRecorderSources.cpp
@@ -542,7 +542,7 @@ ULevelSequence* UTakeRecorderSources::CreateSubSequenceForSource(ULevelSequence*
 			UTakeMetaData* OutTakeMetaData = OutAsset->CopyMetaData(TakeMetaData);
 
 			// Tack on the sub sequence name so that it's unique from the master sequence
-			OutTakeMetaData->SetSlate(TakeMetaData->GetSlate() + TEXT("_") + SubSequenceTrackName);
+			OutTakeMetaData->SetSlate(TakeMetaData->GetSlate() + TEXT("_") + SubSequenceTrackName, false);
 		}
 
 		OutAsset->MarkPackageDirty();
diff --git a/Engine/Plugins/VirtualProduction/Takes/Source/TakesCore/Private/TakesCoreBlueprintLibrary.cpp b/Engine/Plugins/VirtualProduction/Takes/Source/TakesCore/Private/TakesCoreBlueprintLibrary.cpp
--- a/Engine/Plugins/VirtualProduction/Takes/Source/TakesCore/Private/TakesCoreBlueprintLibrary.cpp
+++ b/Engine/Plugins/VirtualProduction/Takes/Source/TakesCore/Private/TakesCoreBlueprintLibrary.cpp
@@ -5,6 +5,12 @@
 #include "LevelSequence.h"
 #include "AssetRegistryModule.h"
 
+namespace
+{
+	static UTakesCoreBlueprintLibrary::FOnTakeRecorderSlateChanged TakeRecorderSlateChanged;
+	static UTakesCoreBlueprintLibrary::FOnTakeRecorderTakeNumberChanged TakeRecorderTakeNumberChanged;
+}
+
 int32 UTakesCoreBlueprintLibrary::ComputeNextTakeNumber(const FString& Slate)
 {
 	int32 MaxTake = 0;
@@ -43,4 +49,26 @@ TArray<FAssetData> UTakesCoreBlueprintLibrary::FindTakes(const FString& Slate, i
 	}
 	
 	return AllAssets;
-}
\ No newline at end of file
+}
+
+void UTakesCoreBlueprintLibrary::SetOnTakeRecorderSlateChanged(FOnTakeRecorderSlateChanged OnTakeRecorderSlateChanged)
+{
+	TakeRecorderSlateChanged = OnTakeRecorderSlateChanged;
+};
+
+void UTakesCoreBlueprintLibrary::SetOnTakeRecorderTakeNumberChanged(FOnTakeRecorderTakeNumberChanged OnTakeRecorderTakeNumberChanged)
+{
+	TakeRecorderTakeNumberChanged = OnTakeRecorderTakeNumberChanged;
+};
+
+void UTakesCoreBlueprintLibrary::OnTakeRecorderSlateChanged(const FString& InSlate)
+{
+	TakeRecorderSlateChanged.ExecuteIfBound(InSlate);
+}
+
+void UTakesCoreBlueprintLibrary::OnTakeRecorderTakeNumberChanged(int32 InTakeNumber)
+{
+	TakeRecorderTakeNumberChanged.ExecuteIfBound(InTakeNumber);
+}
+
+
diff --git a/Engine/Plugins/VirtualProduction/Takes/Source/TakesCore/Public/TakeMetaData.h b/Engine/Plugins/VirtualProduction/Takes/Source/TakesCore/Public/TakeMetaData.h
--- a/Engine/Plugins/VirtualProduction/Takes/Source/TakesCore/Public/TakeMetaData.h
+++ b/Engine/Plugins/VirtualProduction/Takes/Source/TakesCore/Public/TakeMetaData.h
@@ -156,17 +156,19 @@ public:
 
 	/**
 	 * Set the slate for this take and reset its take number to 1
+	 * @param bEmitChanged Whether or not to send a slate changed event
 	 * @note: Only valid for takes that have not been locked
 	 */
 	UFUNCTION(BlueprintCallable, Category="Take")
-	void SetSlate(FString InSlate);
+	void SetSlate(FString InSlate, bool bEmitChanged = true);
 
 	/**
 	 * Set this take's take number. Take numbers are always clamped to be >= 1.
+	 * @param bEmitChanged Whether or not to send a take number changed event
 	 * @note: Only valid for takes that have not been locked
 	 */
 	UFUNCTION(BlueprintCallable, Category="Take")
-	void SetTakeNumber(int32 InTakeNumber);
+	void SetTakeNumber(int32 InTakeNumber, bool bEmitChanged = true);
 
 	/**
 	 * Set this take's timestamp
diff --git a/Engine/Plugins/VirtualProduction/Takes/Source/TakesCore/Public/TakesCoreBlueprintLibrary.h b/Engine/Plugins/VirtualProduction/Takes/Source/TakesCore/Public/TakesCoreBlueprintLibrary.h
--- a/Engine/Plugins/VirtualProduction/Takes/Source/TakesCore/Public/TakesCoreBlueprintLibrary.h
+++ b/Engine/Plugins/VirtualProduction/Takes/Source/TakesCore/Public/TakesCoreBlueprintLibrary.h
@@ -29,4 +29,21 @@ public:
 	 */
 	UFUNCTION(BlueprintCallable, Category="Take Recorder")
 	static TArray<FAssetData> FindTakes(const FString& Slate, int32 TakeNumber = 0);
+
+public:
+
+	DECLARE_DYNAMIC_DELEGATE_OneParam(FOnTakeRecorderSlateChanged, const FString&, Slate);
+	DECLARE_DYNAMIC_DELEGATE_OneParam(FOnTakeRecorderTakeNumberChanged, int32, TakeNumber);
+
+	/** Called when the slate is changed. */
+	UFUNCTION(BlueprintCallable, Category = "Take Recorder")
+	static void SetOnTakeRecorderSlateChanged(FOnTakeRecorderSlateChanged OnTakeRecorderSlateChanged);
+
+	/** Called when the take number is changed. */
+	UFUNCTION(BlueprintCallable, Category = "Take Recorder")
+	static void SetOnTakeRecorderTakeNumberChanged(FOnTakeRecorderTakeNumberChanged OnTakeRecorderTakeNumberChanged);
+
+	static void OnTakeRecorderSlateChanged(const FString& InSlate);
+	static void OnTakeRecorderTakeNumberChanged(int32 InTakeNumber);
+
 };
