Take Recorder: Add a camera cut if there's a camera component (not just if it's a camera actor)

diff --git a/Engine/Plugins/VirtualProduction/Takes/Source/TakeRecorderSources/Private/TakeRecorderActorSource.cpp b/Engine/Plugins/VirtualProduction/Takes/Source/TakeRecorderSources/Private/TakeRecorderActorSource.cpp
--- a/Engine/Plugins/VirtualProduction/Takes/Source/TakeRecorderSources/Private/TakeRecorderActorSource.cpp
+++ b/Engine/Plugins/VirtualProduction/Takes/Source/TakeRecorderSources/Private/TakeRecorderActorSource.cpp
@@ -22,7 +22,7 @@
 #include "GameFramework/Actor.h"
 #include "GameFramework/MovementComponent.h"
 #include "GameFramework/Character.h"
-#include "Camera/CameraActor.h"
+#include "Camera/CameraComponent.h"
 #include "CameraRig_Crane.h"
 #include "CameraRig_Rail.h"
 #include "Components/SkeletalMeshComponent.h"
@@ -636,8 +636,10 @@ TArray<UTakeRecorderSource*> UTakeRecorderActorSource::PostRecording(ULevelSeque
 
 	if (Target.IsValid())
 	{
-		// Automatically add or update the camera cut track if this is a camera
-		if (Target.Get()->IsA<ACameraActor>())
+		// Automatically add or update the camera cut track if there is a camera component
+		AActor* TargetActor = Target.Get();
+
+		if (TargetActor->GetComponentByClass(UCameraComponent::StaticClass()))
 		{
 			FGuid RecordedCameraGuid = GetRecordedActorGuid(Target.Get());
 			FMovieSceneSequenceID RecordedCameraSequenceID = GetLevelSequenceID(Target.Get());
@@ -834,7 +836,7 @@ void UTakeRecorderActorSource::PostEditChangeProperty(FPropertyChangedEvent& Pro
 		
 		TrackTint = FColor(67, 148, 135); 
 		AActor* TargetActor = Target.Get();
-		if (TargetActor && TargetActor->GetClass()->IsChildOf(ACameraActor::StaticClass()))
+		if (TargetActor && TargetActor->GetComponentByClass(UCameraComponent::StaticClass()))
 		{
 			TrackTint = FColor(148, 67, 108);
 		}
@@ -1108,7 +1110,7 @@ FText UTakeRecorderActorSource::GetDisplayTextImpl() const
 FText UTakeRecorderActorSource::GetCategoryTextImpl() const
 {
 	AActor* TargetActor = Target.Get();
-	if (TargetActor && TargetActor->GetClass()->IsChildOf(ACameraActor::StaticClass()))
+	if (TargetActor && TargetActor->GetComponentByClass(UCameraComponent::StaticClass()))
 	{
 		return LOCTEXT("CamerasCategoryLabel", "Cameras");
 	}
