Take Recorder: Force role to ROLE_Authority for client spawnables

diff --git a/Engine/Plugins/VirtualProduction/Takes/Source/TakeRecorderSources/Private/TakeRecorderActorSource.cpp b/Engine/Plugins/VirtualProduction/Takes/Source/TakeRecorderSources/Private/TakeRecorderActorSource.cpp
--- a/Engine/Plugins/VirtualProduction/Takes/Source/TakeRecorderSources/Private/TakeRecorderActorSource.cpp
+++ b/Engine/Plugins/VirtualProduction/Takes/Source/TakeRecorderSources/Private/TakeRecorderActorSource.cpp
@@ -653,6 +653,12 @@ TArray<UTakeRecorderSource*> UTakeRecorderActorSource::PostRecording(ULevelSeque
 		}
 	}
 
+	// Force to authority role in case of capturing replicated actors
+	if (CachedObjectTemplate.IsValid() && !CachedObjectTemplate->HasAuthority())
+	{
+		CachedObjectTemplate->Role = ROLE_Authority;
+	}
+
 	// No longer need to track the Object Template that was created inside the level sequence.
 	CachedObjectBindingGuid = FGuid();
 	CachedObjectTemplate = nullptr;
