Take Recorder: Do not emit slate and take number changed on construction

diff --git a/Engine/Plugins/VirtualProduction/Takes/Source/TakeRecorder/Private/Widgets/STakeRecorderCockpit.cpp b/Engine/Plugins/VirtualProduction/Takes/Source/TakeRecorder/Private/Widgets/STakeRecorderCockpit.cpp
--- a/Engine/Plugins/VirtualProduction/Takes/Source/TakeRecorder/Private/Widgets/STakeRecorderCockpit.cpp
+++ b/Engine/Plugins/VirtualProduction/Takes/Source/TakeRecorder/Private/Widgets/STakeRecorderCockpit.cpp
@@ -560,14 +560,14 @@ void STakeRecorderCockpit::CacheMetaData()
 			FString DefaultSlate = GetDefault<UTakeRecorderProjectSettings>()->Settings.DefaultSlate;
 			if (TransientTakeMetaData->GetSlate() != DefaultSlate)
 			{
-				TransientTakeMetaData->SetSlate(DefaultSlate);
+				TransientTakeMetaData->SetSlate(DefaultSlate, false);
 			}
 
 			// Compute the correct starting take number
 			int32 NextTakeNumber = UTakesCoreBlueprintLibrary::ComputeNextTakeNumber(TransientTakeMetaData->GetSlate());
 			if (TransientTakeMetaData->GetTakeNumber() != NextTakeNumber)
 			{
-				TransientTakeMetaData->SetTakeNumber(NextTakeNumber);
+				TransientTakeMetaData->SetTakeNumber(NextTakeNumber, false);
 			}
 		}
 
