Take Recorder: Added OnStopped

diff --git a/Engine/Plugins/VirtualProduction/Takes/Source/TakeRecorder/Private/Recorder/TakeRecorder.cpp b/Engine/Plugins/VirtualProduction/Takes/Source/TakeRecorder/Private/Recorder/TakeRecorder.cpp
--- a/Engine/Plugins/VirtualProduction/Takes/Source/TakeRecorder/Private/Recorder/TakeRecorder.cpp
+++ b/Engine/Plugins/VirtualProduction/Takes/Source/TakeRecorder/Private/Recorder/TakeRecorder.cpp
@@ -721,6 +721,8 @@ void UTakeRecorder::Stop()
 
 	if (bDidEverStartRecording)
 	{
+		UTakeRecorderBlueprintLibrary::OnTakeRecorderStopped();
+
 		FTakeRecorderSourcesSettings TakeRecorderSourcesSettings;
 		TakeRecorderSourcesSettings.bSaveRecordedAssets = Parameters.User.bSaveRecordedAssets || GEditor == nullptr;
 		TakeRecorderSourcesSettings.bRemoveRedundantTracks = Parameters.User.bRemoveRedundantTracks;
diff --git a/Engine/Plugins/VirtualProduction/Takes/Source/TakeRecorder/Private/Recorder/TakeRecorderBlueprintLibrary.cpp b/Engine/Plugins/VirtualProduction/Takes/Source/TakeRecorder/Private/Recorder/TakeRecorderBlueprintLibrary.cpp
--- a/Engine/Plugins/VirtualProduction/Takes/Source/TakeRecorder/Private/Recorder/TakeRecorderBlueprintLibrary.cpp
+++ b/Engine/Plugins/VirtualProduction/Takes/Source/TakeRecorder/Private/Recorder/TakeRecorderBlueprintLibrary.cpp
@@ -16,6 +16,7 @@ namespace
 	static UTakeRecorderPanel* CurrentTakeRecorderPanel;
 	static UTakeRecorderBlueprintLibrary::FOnTakeRecorderPanelChanged TakeRecorderPanelChanged;
 	static UTakeRecorderBlueprintLibrary::FOnTakeRecorderStarted TakeRecorderStarted;
+	static UTakeRecorderBlueprintLibrary::FOnTakeRecorderStopped TakeRecorderStopped;
 	static UTakeRecorderBlueprintLibrary::FOnTakeRecorderFinished TakeRecorderFinished;
 	static UTakeRecorderBlueprintLibrary::FOnTakeRecorderCancelled TakeRecorderCancelled;
 	static UTakeRecorderBlueprintLibrary::FOnTakeRecorderMarkedFrameAdded TakeRecorderMarkedFrameAdded;
@@ -135,6 +136,11 @@ void UTakeRecorderBlueprintLibrary::SetOnTakeRecorderStarted(FOnTakeRecorderStar
 	TakeRecorderStarted = OnTakeRecorderStarted;
 };
 
+void UTakeRecorderBlueprintLibrary::SetOnTakeRecorderStopped(FOnTakeRecorderStopped OnTakeRecorderStopped)
+{
+	TakeRecorderStopped = OnTakeRecorderStopped;
+};
+
 void UTakeRecorderBlueprintLibrary::SetOnTakeRecorderFinished(FOnTakeRecorderFinished OnTakeRecorderFinished)
 {
 	TakeRecorderFinished = OnTakeRecorderFinished;
@@ -155,6 +161,11 @@ void UTakeRecorderBlueprintLibrary::OnTakeRecorderStarted()
 	TakeRecorderStarted.ExecuteIfBound();
 }
 
+void UTakeRecorderBlueprintLibrary::OnTakeRecorderStopped()
+{
+	TakeRecorderStopped.ExecuteIfBound();
+}
+
 void UTakeRecorderBlueprintLibrary::OnTakeRecorderFinished(ULevelSequence* InSequenceAsset)
 {
 	TakeRecorderFinished.ExecuteIfBound(InSequenceAsset);
diff --git a/Engine/Plugins/VirtualProduction/Takes/Source/TakeRecorder/Public/Recorder/TakeRecorderBlueprintLibrary.h b/Engine/Plugins/VirtualProduction/Takes/Source/TakeRecorder/Public/Recorder/TakeRecorderBlueprintLibrary.h
--- a/Engine/Plugins/VirtualProduction/Takes/Source/TakeRecorder/Public/Recorder/TakeRecorderBlueprintLibrary.h
+++ b/Engine/Plugins/VirtualProduction/Takes/Source/TakeRecorder/Public/Recorder/TakeRecorderBlueprintLibrary.h
@@ -82,6 +82,7 @@ public:
 
 	DECLARE_DYNAMIC_DELEGATE(FOnTakeRecorderPanelChanged);
 	DECLARE_DYNAMIC_DELEGATE(FOnTakeRecorderStarted);
+	DECLARE_DYNAMIC_DELEGATE(FOnTakeRecorderStopped);
 	DECLARE_DYNAMIC_DELEGATE_OneParam(FOnTakeRecorderFinished, ULevelSequence*, SequenceAsset);
 	DECLARE_DYNAMIC_DELEGATE(FOnTakeRecorderCancelled);
 	DECLARE_DYNAMIC_DELEGATE_OneParam(FOnTakeRecorderMarkedFrameAdded, const FMovieSceneMarkedFrame&, MarkedFrame);
@@ -94,6 +95,10 @@ public:
 	UFUNCTION(BlueprintCallable, Category = "Take Recorder")
 	static void SetOnTakeRecorderStarted(FOnTakeRecorderStarted OnTakeRecorderStarted);
 
+	/** Called when take recording is stopped. */
+	UFUNCTION(BlueprintCallable, Category = "Take Recorder")
+	static void SetOnTakeRecorderStopped(FOnTakeRecorderStopped OnTakeRecorderStopped);
+
 	/** Called when take recording finishes. */
 	UFUNCTION(BlueprintCallable, Category = "Take Recorder")
 	static void SetOnTakeRecorderFinished(FOnTakeRecorderFinished OnTakeRecorderFinished);
@@ -107,6 +112,7 @@ public:
 	static void SetOnTakeRecorderMarkedFrameAdded(FOnTakeRecorderMarkedFrameAdded OnTakeRecorderMarkedFrameAdded);
 
 	static void OnTakeRecorderStarted();
+	static void OnTakeRecorderStopped();
 	static void OnTakeRecorderFinished(ULevelSequence* InSequenceAsset);
 	static void OnTakeRecorderCancelled();
 	static void OnTakeRecorderMarkedFrameAdded(const FMovieSceneMarkedFrame& InMarkedFrame);
