Take Recorder: Fix default tracks for when there's no component

diff --git a/Engine/Plugins/VirtualProduction/Takes/Source/TakeTrackRecorders/Public/TrackRecorders/IMovieSceneTrackRecorderHost.h b/Engine/Plugins/VirtualProduction/Takes/Source/TakeTrackRecorders/Public/TrackRecorders/IMovieSceneTrackRecorderHost.h
--- a/Engine/Plugins/VirtualProduction/Takes/Source/TakeTrackRecorders/Public/TrackRecorders/IMovieSceneTrackRecorderHost.h
+++ b/Engine/Plugins/VirtualProduction/Takes/Source/TakeTrackRecorders/Public/TrackRecorders/IMovieSceneTrackRecorderHost.h
@@ -70,6 +70,11 @@ struct FTrackRecorderSettings
 				TArray<FString> ComponentNames;
 				PropertyTrackSetting.ComponentPath.ParseIntoArray(ComponentNames, TEXT("."));
 
+				if (ComponentNames.Num() == 0)
+				{
+					return true;
+				}
+
 				for (const FString& ComponentName : ComponentNames)
 				{
 					if (FindObjectFast<UObject>(InObjectToRecord, *ComponentName) != nullptr)
