Manual Merge From Another Branch.

Make sure animation really starts recording at the correct time, if there's a hitch it may record too late.
Scrubbing with snapping off doesn't evaluate swept frames (events or live link) correctly.
Don't remove empty transform for spawnable since it's no longer true that it will match with template.
Characters no longer animated until we cleared/reset their body channel

When doing animation recording if we want to transfer the root animation on the skeletal mesh over to the transform track, we need to force root lock and then move it over. Force Root Lock(with zero pose) requires transform being identity at first frame. So we zero out the animation as before at first frame, but now correctly handle this zero'd out root transform being applied correctly to not only the skeletal mesh, but it's inverse also getting applied to any attached actors that we are recording to it.

Was already reviewed previously.

diff --git a/Engine/Plugins/VirtualProduction/Takes/Source/TakeRecorderSources/Private/TakeRecorderActorSource.cpp b/Engine/Plugins/VirtualProduction/Takes/Source/TakeRecorderSources/Private/TakeRecorderActorSource.cpp
--- a/Engine/Plugins/VirtualProduction/Takes/Source/TakeRecorderSources/Private/TakeRecorderActorSource.cpp
+++ b/Engine/Plugins/VirtualProduction/Takes/Source/TakeRecorderSources/Private/TakeRecorderActorSource.cpp
@@ -36,8 +36,6 @@
 #include "TrackRecorders/MovieSceneTrackRecorderSettings.h"
 
 #include "Compilation/MovieSceneCompiler.h"
-
-
 DEFINE_LOG_CATEGORY(ActorSerialization);
 
 #define LOCTEXT_NAMESPACE "UTakeRecorderActorSource"
@@ -186,6 +184,8 @@ TArray<UTakeRecorderSource*> UTakeRecorderActorSource::PreRecording(class ULevel
 	{
 		return TArray<UTakeRecorderSource*>();
 	}
+	// We used to do this at the PostRecording but other Actor sources may need to check you recorded an animation so we keep it around
+	TrackRecorders.Empty();
 
 	// Resolve which actor we wish to record 
 	AActor* ActorToRecord = Target.Get();
@@ -628,9 +628,6 @@ TArray<UTakeRecorderSource*> UTakeRecorderActorSource::PostRecording(ULevelSeque
 		SectionRecorder->FinalizeTrack();
 	}
 
-	// Now that the section recorders have placed their data inside the resulting Level Sequence we can release them.
-	TrackRecorders.Empty();
-
 	// Expand the Movie Scene Playback Range to encompass all of the sections now that they've all been created.
 	SequenceRecorderUtils::ExtendSequencePlaybackRange(InSequence);
 
@@ -1472,6 +1469,30 @@ FGuid UTakeRecorderActorSource::GetRecordedActorGuid(class AActor* OtherActor) c
 	return FGuid();
 }
 
+FTransform UTakeRecorderActorSource::GetRecordedActorAnimationInitialRootTransform(class AActor* OtherActor) const
+{
+	UTakeRecorderSources* OwningSources = CastChecked<UTakeRecorderSources>(GetOuter());
+	for (UTakeRecorderSource* Source : OwningSources->GetSources())
+	{
+		if (UTakeRecorderActorSource* ActorSource = Cast<UTakeRecorderActorSource>(Source))
+		{
+			AActor* OtherTarget = ActorSource->Target.Get();
+			if (OtherTarget && OtherActor && (OtherTarget == OtherActor || OtherTarget->GetName() == OtherActor->GetName()))
+			{
+				for (UMovieSceneTrackRecorder* TrackRecorder : ActorSource->TrackRecorders)
+				{
+					if (TrackRecorder->IsA<UMovieSceneAnimationTrackRecorder>())
+					{
+						return Cast<UMovieSceneAnimationTrackRecorder>(TrackRecorder)->GetInitialRootTransform();
+					}
+				}
+
+			}
+		}
+	}
+	return FTransform::Identity;
+}
+
 FMovieSceneSequenceID UTakeRecorderActorSource::GetLevelSequenceID(class AActor* OtherActor)
 {
 	FMovieSceneSequenceID OutSequenceID = MovieSceneSequenceID::Root;
diff --git a/Engine/Plugins/VirtualProduction/Takes/Source/TakeRecorderSources/Private/TakeRecorderActorSource.h b/Engine/Plugins/VirtualProduction/Takes/Source/TakeRecorderSources/Private/TakeRecorderActorSource.h
--- a/Engine/Plugins/VirtualProduction/Takes/Source/TakeRecorderSources/Private/TakeRecorderActorSource.h
+++ b/Engine/Plugins/VirtualProduction/Takes/Source/TakeRecorderSources/Private/TakeRecorderActorSource.h
@@ -236,6 +236,9 @@ protected:
 	FMovieSceneSequenceID GetLevelSequenceID(class AActor* OtherActor) override;
 	/** Returns generic track recorder settings */
 	FTrackRecorderSettings GetTrackRecorderSettings() const override;
+	/** Returns offset that may get set when recording a skeletal mesh animation. Needed to correctly transform attached children*/
+	FTransform GetRecordedActorAnimationInitialRootTransform(class AActor* OtherActor) const override;
+
 	// ~IMovieSceneTrackRecorderHost Interface
 	
 	/** Initializes an instance of the specified class if we don't already have it in our Settings array. */
diff --git a/Engine/Plugins/VirtualProduction/Takes/Source/TakeTrackRecorders/Private/MovieScene3DTransformTrackRecorder.cpp b/Engine/Plugins/VirtualProduction/Takes/Source/TakeTrackRecorders/Private/MovieScene3DTransformTrackRecorder.cpp
--- a/Engine/Plugins/VirtualProduction/Takes/Source/TakeTrackRecorders/Private/MovieScene3DTransformTrackRecorder.cpp
+++ b/Engine/Plugins/VirtualProduction/Takes/Source/TakeTrackRecorders/Private/MovieScene3DTransformTrackRecorder.cpp
@@ -11,6 +11,8 @@
 #include "Channels/MovieSceneChannelProxy.h"
 #include "Components/SceneComponent.h"
 #include "Components/SkeletalMeshComponent.h"
+#include "SequenceRecorderUtils.h"
+
 
 DEFINE_LOG_CATEGORY(TransformSerialization);
 
@@ -111,6 +113,45 @@ void UMovieScene3DTransformTrackRecorder::FinalizeTrackImpl()
  
  	SlowTask.EnterProgressFrame();
  
+	FTransform InvParentAnimationRootTransform = FTransform::Identity;
+	FName SocketName;
+	FName ComponentName;
+	AActor* ActorToRecord = Cast<AActor>(ObjectToRecord.Get());
+	if (ActorToRecord)
+	{
+		AActor* AttachedToActor = SequenceRecorderUtils::GetAttachment(ActorToRecord, SocketName, ComponentName);
+		if (AttachedToActor)
+		{
+			InvParentAnimationRootTransform = OwningTakeRecorderSource->GetRecordedActorAnimationInitialRootTransform(AttachedToActor).Inverse();
+			if (!InvParentAnimationRootTransform.Equals(FTransform::Identity))
+			{
+				if (DefaultTransform.IsSet())
+				{
+					FTransform DT = DefaultTransform.GetValue();
+					DT = DT * InvParentAnimationRootTransform;
+					DefaultTransform = DT;
+					FVector Translation = DT.GetTranslation();
+					FVector EulerRotation = DT.GetRotation().Rotator().Euler();
+					FVector Scale = DT.GetScale3D();
+					TArrayView<FMovieSceneFloatChannel*> FloatChannels = MovieSceneSection->GetChannelProxy().GetChannels<FMovieSceneFloatChannel>();
+					FloatChannels[0]->SetDefault(Translation.X);
+					FloatChannels[1]->SetDefault(Translation.Y);
+					FloatChannels[2]->SetDefault(Translation.Z);
+					FloatChannels[3]->SetDefault(EulerRotation.X);
+					FloatChannels[4]->SetDefault(EulerRotation.Y);
+					FloatChannels[5]->SetDefault(EulerRotation.Z);
+					FloatChannels[6]->SetDefault(Scale.X);
+					FloatChannels[7]->SetDefault(Scale.Y);
+					FloatChannels[8]->SetDefault(Scale.Z);
+				}
+				if (BufferedTransforms.Times.Num() > 0)
+				{
+					BufferedTransforms.PostMultTransform(InvParentAnimationRootTransform);
+				}
+			}
+		}
+	}
+
  	// Try to 're-wind' rotations that look like axis flips
  	// We need to do this as a post-process because the recorder cant reliably access 'wound' rotations:
  	// - Net quantize may use quaternions.
@@ -242,15 +283,7 @@ void UMovieScene3DTransformTrackRecorder::FinalizeTrackImpl()
 				FloatChannels[7]->SetDefault(FirstTransform.GetScale3D()[1]);
 				FloatChannels[8]->SetDefault(FirstTransform.GetScale3D()[2]);
 
-				// The section can be removed if this is a spawnable since the spawnable template should have the same default values
-				if (!TrackRecorderSettings.bRecordToPossessable && MovieSceneTrack.IsValid())
-				{
-					FString ObjectToRecordName = ObjectToRecord.IsValid() ? ObjectToRecord->GetName() : TEXT("Unnamed_Actor");
-					UE_LOG(LogTakesCore, Log, TEXT("Removed unused track (%s) for (%s)"), *MovieSceneTrack->GetTrackName().ToString(), *ObjectToRecordName);
-
-					MovieSceneTrack->RemoveSection(*MovieSceneSection);
-					MovieScene->RemoveTrack(*MovieSceneTrack);
-				}
+				//no longer remove spawnable transform tracks since they may not match the template
 			}
  		}
  	}
@@ -349,6 +382,7 @@ void UMovieScene3DTransformTrackRecorder::SetUpDefaultTransform()
 	FloatChannels[7]->SetDefault(Scale.Y);
 	FloatChannels[8]->SetDefault(Scale.Z);
 }
+
 bool UMovieScene3DTransformTrackRecorder::ResolveTransformToRecord(FTransform& OutTransform)
 {
  	if(USceneComponent* SceneComponent = Cast<USceneComponent>(ObjectToRecord.Get()))
@@ -480,7 +514,12 @@ void UMovieScene3DTransformTrackRecorder::PostProcessAnimationData(UMovieSceneAn
 
 				// we may need to offset the transform here if the animation was not recorded on the root component
 				FTransform InvComponentTransform = AnimTrackRecorder->GetComponentTransform().Inverse();
+				FTransform InitialRootTransform = AnimTrackRecorder->GetInitialRootTransform();
+				if (DefaultTransform.IsSet())
+				{
+					DefaultTransform = InitialRootTransform * DefaultTransform.GetValue();
 
+				}
 				const FRawAnimSequenceTrack& RawTrack = AnimSequence->GetRawAnimationData()[RootIndex];
 				const int32 KeyCount = FMath::Max(FMath::Max(RawTrack.PosKeys.Num(), RawTrack.RotKeys.Num()), RawTrack.ScaleKeys.Num());
 				for (int32 KeyIndex = 0; KeyIndex < KeyCount; KeyIndex++)
@@ -514,7 +553,12 @@ void UMovieScene3DTransformTrackRecorder::PostProcessAnimationData(UMovieSceneAn
 					}
 
 					FFrameNumber AnimationFrame = (AnimSequence->GetTimeAtFrame(KeyIndex) * TickResolution).FloorToFrame();
-					AnimationKeys.Add(InvComponentTransform * Transform * Relative, StartTime + AnimationFrame);
+					//TODO Fix this.
+					//this works for tests but not props
+					AnimationKeys.Add(InvComponentTransform * InitialRootTransform * Transform * Relative, StartTime + AnimationFrame);
+
+					//this works for props but not tests
+					//AnimationKeys.Add(InvComponentTransform * Transform * InitialRootTransform * Relative, StartTime + AnimationFrame);
 				}
 			}
 		}
diff --git a/Engine/Plugins/VirtualProduction/Takes/Source/TakeTrackRecorders/Private/MovieSceneAnimationTrackRecorder.cpp b/Engine/Plugins/VirtualProduction/Takes/Source/TakeTrackRecorders/Private/MovieSceneAnimationTrackRecorder.cpp
--- a/Engine/Plugins/VirtualProduction/Takes/Source/TakeTrackRecorders/Private/MovieSceneAnimationTrackRecorder.cpp
+++ b/Engine/Plugins/VirtualProduction/Takes/Source/TakeTrackRecorders/Private/MovieSceneAnimationTrackRecorder.cpp
@@ -173,6 +173,7 @@ void UMovieSceneAnimationTrackRecorder::StopRecordingImpl()
 	{
 		// Legacy Animation Recorder allowed recording into an animation asset directly and not creating an movie section
 		const bool bShowAnimationAssetCreatedToast = false;
+		InitialRootTransform = FAnimationRecorderManager::Get().GetInitialRootTransform(SkeletalMeshComponent.Get());
 		FAnimationRecorderManager::Get().StopRecordingAnimation(SkeletalMeshComponent.Get(), bShowAnimationAssetCreatedToast);
 	}
 }
@@ -201,7 +202,18 @@ void UMovieSceneAnimationTrackRecorder::RecordSampleImpl(const FQualifiedFrameTi
 	//  Note we wait for first tick so that we can make sure all of the attach tracks are set up .
 	if (!bAnimationRecorderCreated)
 	{
-		
+		//Reset the start times based upon when the animation really starts.
+		if (MovieSceneSection.IsValid())
+		{
+			MovieSceneSection->TimecodeSource = FMovieSceneTimecodeSource(FApp::GetTimecode());
+			FFrameRate   TickResolution = MovieSceneSection->GetTypedOuter<UMovieScene>()->GetTickResolution();
+			FFrameNumber CurrentFrame = CurrentTime.ConvertTo(TickResolution).FloorToFrame();
+
+			// Ensure we're expanded to at least the next frame so that we don't set the start past the end
+			// when we set the first frame.
+			MovieSceneSection->ExpandToFrame(CurrentFrame + FFrameNumber(1));
+			MovieSceneSection->SetStartFrame(TRangeBound<FFrameNumber>::Inclusive(CurrentFrame));
+		}
 		bAnimationRecorderCreated = true;
 		AActor* Actor = nullptr;
 		SkeletalMeshComponent = CastChecked<USkeletalMeshComponent>(ObjectToRecord.Get());
@@ -212,11 +224,9 @@ void UMovieSceneAnimationTrackRecorder::RecordSampleImpl(const FQualifiedFrameTi
 		//In Sequence Recorder this would be done via checking if the component was dynamically created, due to changes in how the take recorder handles this, it no longer 
 		//possible so it seems if it's native do root, otherwise, don't.  
 		bool bRemoveRootAnimation = SkeletalMeshComponent->CreationMethod != EComponentCreationMethod::Native ? false : true;
-
 		// Need to pass this up to the settings since it's used later to force root lock and transfer root from animation to transform
 		UMovieSceneAnimationTrackRecorderSettings* AnimSettings = CastChecked<UMovieSceneAnimationTrackRecorderSettings>(Settings.Get());
 		AnimSettings->bRemoveRootAnimation = bRemoveRootAnimation;
-
 		//If not removing root we also don't record in world space ( not totally sure if it matters but matching up with Sequence Recorder)
 		bool bRecordInWorldSpace = bRemoveRootAnimation == false ? false : true;
 		
diff --git a/Engine/Plugins/VirtualProduction/Takes/Source/TakeTrackRecorders/Public/TrackRecorders/IMovieSceneTrackRecorderHost.h b/Engine/Plugins/VirtualProduction/Takes/Source/TakeTrackRecorders/Public/TrackRecorders/IMovieSceneTrackRecorderHost.h
--- a/Engine/Plugins/VirtualProduction/Takes/Source/TakeTrackRecorders/Public/TrackRecorders/IMovieSceneTrackRecorderHost.h
+++ b/Engine/Plugins/VirtualProduction/Takes/Source/TakeTrackRecorders/Public/TrackRecorders/IMovieSceneTrackRecorderHost.h
@@ -167,4 +167,12 @@ public:
 	 * 
 	 */
 	virtual FTrackRecorderSettings GetTrackRecorderSettings() const = 0;
+
+	/** Returns original transform that may get set when recording a skeletal mesh animation. Needed to correctly transform attached children.
+	* @param OtherActor Actor to look for.
+	* @return The FTransform applied to the actor when recording animation, that the attached actor may need to apply to themselves to orient correctly.
+	*/
+	virtual FTransform GetRecordedActorAnimationInitialRootTransform(class AActor* OtherActor) const = 0;
+
+	/**/
 };
diff --git a/Engine/Plugins/VirtualProduction/Takes/Source/TakeTrackRecorders/Public/TrackRecorders/MovieScene3DTransformTrackRecorder.h b/Engine/Plugins/VirtualProduction/Takes/Source/TakeTrackRecorders/Public/TrackRecorders/MovieScene3DTransformTrackRecorder.h
--- a/Engine/Plugins/VirtualProduction/Takes/Source/TakeTrackRecorders/Public/TrackRecorders/MovieScene3DTransformTrackRecorder.h
+++ b/Engine/Plugins/VirtualProduction/Takes/Source/TakeTrackRecorders/Public/TrackRecorders/MovieScene3DTransformTrackRecorder.h
@@ -105,6 +105,7 @@ struct FBufferedTransformKeys
 		}
 	}
 
+
 	void GetValueFromIndex(int32 CurIndex, FVector& OutLocation, FQuat& OutQuat, FVector& OutScale)
 	{
 		OutLocation.X = LocationX[CurIndex];
@@ -177,6 +178,34 @@ struct FBufferedTransformKeys
 			}
 		}
 	}
+
+	void PostMultTransform(const FTransform &Transform)
+	{
+		CreateCachedQuats(); //do stuff in quat space to avoid flips
+		FVector OurLocation, OurScale;
+		FQuat OurQuat;
+		FTransform OurTransform;
+		for (int32 CurIndex = 0; CurIndex < Times.Num(); ++CurIndex)
+		{
+			GetValueFromIndex(CurIndex, OurLocation, OurQuat, OurScale);
+			OurTransform = FTransform(OurQuat, OurLocation, OurScale);
+
+			OurTransform = OurTransform * Transform;
+			//OurTransform = OurTransform * Transform.Inverse();
+			LocationX[CurIndex] = OurTransform.GetTranslation().X;
+			LocationY[CurIndex] = OurTransform.GetTranslation().Y;
+			LocationZ[CurIndex] = OurTransform.GetTranslation().Z;
+
+			FRotator WoundRoation = OurTransform.Rotator();
+			RotationX[CurIndex] = WoundRoation.Roll;
+			RotationY[CurIndex] = WoundRoation.Pitch;
+			RotationZ[CurIndex] = WoundRoation.Yaw;
+
+			ScaleX[CurIndex] = OurTransform.GetScale3D().X;
+			ScaleY[CurIndex] = OurTransform.GetScale3D().Y;
+			ScaleZ[CurIndex] = OurTransform.GetScale3D().Z;
+		}
+	}
 	/**
 	Collapse the passed in Additve onto myself, returning
 	a new FBufferedTransformKeys
diff --git a/Engine/Plugins/VirtualProduction/Takes/Source/TakeTrackRecorders/Public/TrackRecorders/MovieSceneAnimationTrackRecorder.h b/Engine/Plugins/VirtualProduction/Takes/Source/TakeTrackRecorders/Public/TrackRecorders/MovieSceneAnimationTrackRecorder.h
--- a/Engine/Plugins/VirtualProduction/Takes/Source/TakeTrackRecorders/Public/TrackRecorders/MovieSceneAnimationTrackRecorder.h
+++ b/Engine/Plugins/VirtualProduction/Takes/Source/TakeTrackRecorders/Public/TrackRecorders/MovieSceneAnimationTrackRecorder.h
@@ -61,6 +61,7 @@ public:
 	USkeletalMesh* GetSkeletalMesh() const { return SkeletalMesh.Get(); }
 	USkeletalMeshComponent* GetSkeletalMeshComponent() const { return SkeletalMeshComponent.Get(); }
 	const FTransform& GetComponentTransform() const { return ComponentTransform; }
+	const FTransform& GetInitialRootTransform() const { return InitialRootTransform; }
 
 private:
 	bool ResolveTransformToRecord(FTransform& TransformToRecord);
@@ -78,6 +79,9 @@ private:
 	/** Local transform of the component we are recording */
 	FTransform ComponentTransform;
 
+	/** Inverse we are using to zero out root motion */
+	FTransform InitialRootTransform;
+
 	bool bAnimationRecorderCreated; 
 	/**Serializer */
 	FAnimationSerializer AnimationSerializer;
diff --git a/Engine/Source/Editor/SequenceRecorder/Private/AnimationRecorder.cpp b/Engine/Source/Editor/SequenceRecorder/Private/AnimationRecorder.cpp
--- a/Engine/Source/Editor/SequenceRecorder/Private/AnimationRecorder.cpp
+++ b/Engine/Source/Editor/SequenceRecorder/Private/AnimationRecorder.cpp
@@ -576,6 +576,7 @@ bool FAnimationRecorder::Record(USkeletalMeshComponent* Component, FTransform co
 							// we make about root motion use are incorrect.
 							// NEW. But we don't do this if there is just one root bone. This has come up with recording
 							// single bone props and cameras.
+							InitialRootTransform = LocalTransform;
 							InvInitialRootTransform = LocalTransform.Inverse();
 						}
 						else
@@ -605,7 +606,7 @@ bool FAnimationRecorder::Record(USkeletalMeshComponent* Component, FTransform co
 					if (ParentIndex == SkeletonRootIndex)
 					{
 						// Remove initial root transform
-						LocalTransform.SetToRelativeTransform(SpacesBases[ParentIndex] * InvInitialRootTransform);
+						LocalTransform.SetToRelativeTransform(SpacesBases[ParentIndex]);
 					}
 					else
 					{
@@ -617,7 +618,6 @@ bool FAnimationRecorder::Record(USkeletalMeshComponent* Component, FTransform co
 				{
 					// Remove initial root transform
 					LocalTransform *= InvInitialRootTransform;
-					//LocalTransform = InvInitialRootTransform * LocalTransform;
 
 					if (bRecordLocalToWorld)
 					{
@@ -1030,6 +1030,18 @@ float FAnimationRecorderManager::GetCurrentRecordingTime(USkeletalMeshComponent*
 	return 0.0f;
 }
 
+const FTransform&  FAnimationRecorderManager::GetInitialRootTransform(USkeletalMeshComponent* Component) const
+{
+	for (const FAnimRecorderInstance& Instance : RecorderInstances)
+	{
+		if (Instance.SkelComp == Component)
+		{
+			return Instance.Recorder->GetInitialRootTransform();
+		}
+	}
+	return FTransform::Identity;
+}
+
 void FAnimationRecorderManager::StopRecordingAnimation(USkeletalMeshComponent* Component, bool bShowMessage)
 {
 	for (int32 Idx = 0; Idx < RecorderInstances.Num(); ++Idx)
diff --git a/Engine/Source/Editor/SequenceRecorder/Public/AnimationRecorder.h b/Engine/Source/Editor/SequenceRecorder/Public/AnimationRecorder.h
--- a/Engine/Source/Editor/SequenceRecorder/Public/AnimationRecorder.h
+++ b/Engine/Source/Editor/SequenceRecorder/Public/AnimationRecorder.h
@@ -38,6 +38,7 @@ private:
 	FBlendedHeapCurve PreviousAnimCurves;
 	FTransform PreviousComponentToWorld;
 	FTransform InvInitialRootTransform;
+	FTransform InitialRootTransform;
 	int32 SkeletonRootIndex;
 
 	/** Array of currently active notifies that have duration */
@@ -75,6 +76,8 @@ public:
 
 	bool SetAnimCompressionScheme(TSubclassOf<class UAnimCompress> SchemeClass);
 
+	const FTransform& GetInitialRootTransform() const { return InitialRootTransform; }
+
 	/** If true, it will record root to include LocalToWorld */
 	uint8 bRecordLocalToWorld :1;
 	/** If true, asset will be saved to disk after recording. If false, asset will remain in mem and can be manually saved. */
@@ -167,6 +170,7 @@ public:
 	float GetCurrentRecordingTime(USkeletalMeshComponent* Component);
 	void StopRecordingAnimation(USkeletalMeshComponent* Component, bool bShowMessage = true);
 	void StopRecordingAllAnimations();
+	const FTransform& GetInitialRootTransform(USkeletalMeshComponent* Component) const;
 
 	void Tick(float DeltaTime);
 
diff --git a/Engine/Source/Runtime/MovieScene/Private/Evaluation/MovieScenePlayback.cpp b/Engine/Source/Runtime/MovieScene/Private/Evaluation/MovieScenePlayback.cpp
--- a/Engine/Source/Runtime/MovieScene/Private/Evaluation/MovieScenePlayback.cpp
+++ b/Engine/Source/Runtime/MovieScene/Private/Evaluation/MovieScenePlayback.cpp
@@ -77,6 +77,17 @@ TRange<FFrameNumber> FMovieSceneEvaluationRange::GetTraversedFrameNumberRange()
 TRange<FFrameNumber> FMovieSceneEvaluationRange::TimeRangeToNumberRange(const TRange<FFrameTime>& InFrameTimeRange)
 {
 	TRange<FFrameNumber> FrameNumberRange;
+	TOptional<FFrameTime> UpperTime;
+	if (!InFrameTimeRange.GetUpperBound().IsOpen())
+	{
+		UpperTime = InFrameTimeRange.GetUpperBoundValue();
+
+		FrameNumberRange.SetUpperBound(
+			InFrameTimeRange.GetUpperBound().IsExclusive()
+			? TRangeBound<FFrameNumber>::Exclusive(UpperTime.GetValue().FrameNumber)
+			: TRangeBound<FFrameNumber>::Inclusive(UpperTime.GetValue().FrameNumber)
+		);
+	}
 
 	if (!InFrameTimeRange.GetLowerBound().IsOpen())
 	{
@@ -84,21 +95,13 @@ TRange<FFrameNumber> FMovieSceneEvaluationRange::TimeRangeToNumberRange(const TR
 		// If there is a sub frame on the start time, we're actually beyond that frame number, so it needs incrementing
 		if (LowerTime.GetSubFrame() != 0.f || InFrameTimeRange.GetLowerBound().IsExclusive())
 		{
-			LowerTime.FrameNumber = LowerTime.FrameNumber + 1;
+			LowerTime.FrameNumber = (!UpperTime.IsSet() || LowerTime.FrameNumber < UpperTime.GetValue().FrameNumber) ?
+				LowerTime.FrameNumber + 1 : LowerTime.FrameNumber;
+
 		}
 		FrameNumberRange.SetLowerBound(TRangeBound<FFrameNumber>::Inclusive(LowerTime.FrameNumber));
 	}
 
-	if (!InFrameTimeRange.GetUpperBound().IsOpen())
-	{
-		const FFrameTime UpperTime = InFrameTimeRange.GetUpperBoundValue();
-
-		FrameNumberRange.SetUpperBound(
-			InFrameTimeRange.GetUpperBound().IsExclusive()
-			? TRangeBound<FFrameNumber>::Exclusive(UpperTime.FrameNumber)
-			: TRangeBound<FFrameNumber>::Inclusive(UpperTime.FrameNumber)
-		);
-	}
 
 	return FrameNumberRange;
 }
diff --git a/Engine/Source/Runtime/MovieSceneTracks/Private/Evaluation/MovieSceneSkeletalAnimationTemplate.cpp b/Engine/Source/Runtime/MovieSceneTracks/Private/Evaluation/MovieSceneSkeletalAnimationTemplate.cpp
--- a/Engine/Source/Runtime/MovieSceneTracks/Private/Evaluation/MovieSceneSkeletalAnimationTemplate.cpp
+++ b/Engine/Source/Runtime/MovieSceneTracks/Private/Evaluation/MovieSceneSkeletalAnimationTemplate.cpp
@@ -11,6 +11,8 @@
 #include "Evaluation/MovieSceneEvaluation.h"
 #include "IMovieScenePlayer.h"
 #include "UObject/ObjectKey.h"
+#include "UObject/StrongObjectPtr.h"
+
 
 bool ShouldUsePreviewPlayback(IMovieScenePlayer& Player, UObject& RuntimeObject)
 {
@@ -79,6 +81,8 @@ struct FPreAnimatedAnimationTokenProducer : IMovieScenePreAnimatedTokenProducer
 				// Cache this object's current update flag and animation mode
 				VisibilityBasedAnimTickOption = InComponent->VisibilityBasedAnimTickOption;
 				AnimationMode = InComponent->GetAnimationMode();
+				CachedAnimInstance.Reset(InComponent->AnimScriptInstance);
+
 #if WITH_EDITOR
 				bUpdateAnimationInEditor = InComponent->GetUpdateAnimationInEditor();
 #endif
@@ -104,7 +108,6 @@ struct FPreAnimatedAnimationTokenProducer : IMovieScenePreAnimatedTokenProducer
 				}
 
 				UAnimSequencerInstance::UnbindFromSkeletalMeshComponent(Component);
-
 				// Reset the mesh component update flag and animation mode to what they were before we animated the object
 				Component->VisibilityBasedAnimTickOption = VisibilityBasedAnimTickOption;
 				if (Component->GetAnimationMode() != AnimationMode)
@@ -113,6 +116,11 @@ struct FPreAnimatedAnimationTokenProducer : IMovieScenePreAnimatedTokenProducer
 					// if we're using same anim blueprint, we don't want to keep reinitializing it. 
 					Component->SetAnimationMode(AnimationMode);
 				}
+				if (CachedAnimInstance.Get())
+				{
+					Component->AnimScriptInstance = CachedAnimInstance.Get();
+					CachedAnimInstance.Reset();
+				}
 #if WITH_EDITOR
 				Component->SetUpdateAnimationInEditor(bUpdateAnimationInEditor);
 #endif
@@ -120,6 +128,8 @@ struct FPreAnimatedAnimationTokenProducer : IMovieScenePreAnimatedTokenProducer
 
 			EVisibilityBasedAnimTickOption VisibilityBasedAnimTickOption;
 			EAnimationMode::Type AnimationMode;
+			TStrongObjectPtr<UAnimInstance> CachedAnimInstance;
+
 #if WITH_EDITOR
 			bool bUpdateAnimationInEditor;
 #endif
