Take Recorder: Added gc guard for take recorder

diff --git a/Engine/Plugins/VirtualProduction/Takes/Source/TakeRecorder/Private/Recorder/TakeRecorder.cpp b/Engine/Plugins/VirtualProduction/Takes/Source/TakeRecorder/Private/Recorder/TakeRecorder.cpp
--- a/Engine/Plugins/VirtualProduction/Takes/Source/TakeRecorder/Private/Recorder/TakeRecorder.cpp
+++ b/Engine/Plugins/VirtualProduction/Takes/Source/TakeRecorder/Private/Recorder/TakeRecorder.cpp
@@ -26,6 +26,7 @@
 #include "Editor.h"
 #include "Toolkits/AssetEditorManager.h"
 #include "ObjectTools.h"
+#include "UObject/GCObjectScopeGuard.h"
 
 // Slate includes
 #include "Framework/Notifications/NotificationManager.h"
@@ -328,13 +329,13 @@ public:
 FTickableTakeRecorder TickableTakeRecorder;
 
 // Static members of UTakeRecorder
-UTakeRecorder*          UTakeRecorder::CurrentRecorder = nullptr;
+TStrongObjectPtr<UTakeRecorder> CurrentRecorder;
 FOnTakeRecordingInitialized UTakeRecorder::OnRecordingInitializedEvent;
 
 // Static functions for UTakeRecorder
 UTakeRecorder* UTakeRecorder::GetActiveRecorder()
 {
-	return CurrentRecorder;
+	return CurrentRecorder.Get();
 }
 
 FOnTakeRecordingInitialized& UTakeRecorder::OnRecordingInitialized()
@@ -344,13 +345,13 @@ FOnTakeRecordingInitialized& UTakeRecorder::OnRecordingInitialized()
 
 bool UTakeRecorder::SetActiveRecorder(UTakeRecorder* NewActiveRecorder)
 {
-	if (CurrentRecorder)
+	if (CurrentRecorder.IsValid())
 	{
 		return false;
 	}
 
-	CurrentRecorder = NewActiveRecorder;
-	TickableTakeRecorder.WeakRecorder = CurrentRecorder;
+	CurrentRecorder.Reset(NewActiveRecorder);
+	TickableTakeRecorder.WeakRecorder = CurrentRecorder.Get();
 	OnRecordingInitializedEvent.Broadcast(NewActiveRecorder);
 	return true;
 }
@@ -367,6 +368,8 @@ UTakeRecorder::UTakeRecorder(const FObjectInitializer& ObjInit)
 
 bool UTakeRecorder::Initialize(ULevelSequence* LevelSequenceBase, UTakeRecorderSources* Sources, UTakeMetaData* MetaData, const FTakeRecorderParameters& InParameters, FText* OutError)
 {
+	FGCObjectScopeGuard GCGuard(this);
+
 	if (GetActiveRecorder())
 	{
 		if (OutError)
@@ -775,9 +778,9 @@ void UTakeRecorder::Stop()
 	OnStopCleanup.Reset();
 
 	// reset the current recorder and stop us from being ticked
-	if (CurrentRecorder == this)
+	if (CurrentRecorder.Get() == this)
 	{
-		CurrentRecorder = nullptr;
+		CurrentRecorder.Reset();
 		TickableTakeRecorder.WeakRecorder = nullptr;
 
 		if (bDidEverStartRecording)
diff --git a/Engine/Plugins/VirtualProduction/Takes/Source/TakeRecorder/Public/Recorder/TakeRecorder.h b/Engine/Plugins/VirtualProduction/Takes/Source/TakeRecorder/Public/Recorder/TakeRecorder.h
--- a/Engine/Plugins/VirtualProduction/Takes/Source/TakeRecorder/Public/Recorder/TakeRecorder.h
+++ b/Engine/Plugins/VirtualProduction/Takes/Source/TakeRecorder/Public/Recorder/TakeRecorder.h
@@ -3,6 +3,7 @@
 #pragma once
 
 #include "UObject/Object.h"
+#include "UObject/StrongObjectPtr.h"
 #include "ISequencer.h"
 #include "Recorder/TakeRecorderParameters.h"
 #include "Serializers/MovieSceneManifestSerialization.h"
@@ -209,9 +210,6 @@ private:
 	 */
 	static bool SetActiveRecorder(UTakeRecorder* NewActiveRecorder);
 
-	/** A pointer to the currently active recorder */
-	static UTakeRecorder* CurrentRecorder;
-
 	/** Event to trigger when a new recording is initialized */
 	static FOnTakeRecordingInitialized OnRecordingInitializedEvent;
 
