Don't record if time doesn't increase, due to 0.0 DeltaTime or timecode frame not increasing, in order to avoid duplicate keys.

diff --git a/Engine/Plugins/VirtualProduction/Takes/Source/TakesCore/Private/TakeRecorderSources.cpp b/Engine/Plugins/VirtualProduction/Takes/Source/TakesCore/Private/TakeRecorderSources.cpp
--- a/Engine/Plugins/VirtualProduction/Takes/Source/TakesCore/Private/TakeRecorderSources.cpp
+++ b/Engine/Plugins/VirtualProduction/Takes/Source/TakesCore/Private/TakeRecorderSources.cpp
@@ -282,6 +282,7 @@ void UTakeRecorderSources::StartRecording(class ULevelSequence* InSequence, FMan
 
 	bIsRecording = true;
 	TimeSinceRecordingStarted = 0.f;
+	LastTimecodeFrameNumber.Reset();
 	TargetLevelSequenceTickResolution = InSequence->GetMovieScene()->GetTickResolution();
 
 	FTimecode TimecodeSource = FApp::GetTimecode();
@@ -296,19 +297,31 @@ FFrameTime UTakeRecorderSources::TickRecording(class ULevelSequence* InSequence,
 	bool bHasValidTimecodeSource;
 	FQualifiedFrameTime FrameTime = GetCurrentRecordingFrameTime(CurrentTimecode,bHasValidTimecodeSource);
 	FQualifiedFrameTime SourceFrameTime(FrameTime);
+	bool bTimeIncremented = DeltaTime > 0.0f;
 	if (bHasValidTimecodeSource)
 	{
 		//We leave this ins timecode frame rate since the sources convert it later (cbb and faster to do it here, we actually do it below
 		//for showiung the time
 		FFrameNumber SeqStartFrameTime = StartRecordingTimecodeSource.ToFrameNumber(FApp::GetTimecodeFrameRate());
 		SourceFrameTime.Time.FrameNumber = SourceFrameTime.Time.FrameNumber - SeqStartFrameTime;
+		if (LastTimecodeFrameNumber.IsSet())
+		{
+			if (LastTimecodeFrameNumber == SourceFrameTime.Time.FrameNumber)
+			{
+				bTimeIncremented = false;
+			}
+		}
+		LastTimecodeFrameNumber = SourceFrameTime.Time.FrameNumber;
 	}
 
-	for (auto Source : Sources)
+	if (bTimeIncremented) //only record if time incremented, may not with timecode providers with low frame rates
 	{
-		if (Source->bEnabled)
+		for (auto Source : Sources)
 		{
-			Source->TickRecording(SourceFrameTime);
+			if (Source->bEnabled)
+			{
+				Source->TickRecording(SourceFrameTime);
+			}
 		}
 	}
 
@@ -404,7 +417,7 @@ void UTakeRecorderSources::StopRecording(class ULevelSequence* InSequence, FTake
 {
 	bIsRecording = false;
 	TimeSinceRecordingStarted = 0.f;
-
+	LastTimecodeFrameNumber.Reset();
 	for (auto Source : Sources)
 	{
 		if (Source->bEnabled)
diff --git a/Engine/Plugins/VirtualProduction/Takes/Source/TakesCore/Public/TakeRecorderSources.h b/Engine/Plugins/VirtualProduction/Takes/Source/TakesCore/Public/TakeRecorderSources.h
--- a/Engine/Plugins/VirtualProduction/Takes/Source/TakesCore/Public/TakeRecorderSources.h
+++ b/Engine/Plugins/VirtualProduction/Takes/Source/TakesCore/Public/TakeRecorderSources.h
@@ -212,5 +212,8 @@ private:
 
 	/** Timecode time at start of recording */
 	FTimecode StartRecordingTimecodeSource;
+
+	/** Last Timecode Frame Number, used to avoid recording same time twice*/
+	TOptional<FFrameNumber> LastTimecodeFrameNumber;
 };
 
