Take Recorder: Set infinite length when recording starts, rather than when initializing the level sequence

diff --git a/Engine/Plugins/VirtualProduction/Takes/Source/TakesCore/Private/TakePreset.cpp b/Engine/Plugins/VirtualProduction/Takes/Source/TakesCore/Private/TakePreset.cpp
--- a/Engine/Plugins/VirtualProduction/Takes/Source/TakesCore/Private/TakePreset.cpp
+++ b/Engine/Plugins/VirtualProduction/Takes/Source/TakesCore/Private/TakePreset.cpp
@@ -2,6 +2,7 @@
 
 #include "TakePreset.h"
 #include "LevelSequence.h"
+#include "MovieSceneToolsProjectSettings.h"
 #include "UObject/Package.h"
 
 UTakePreset::UTakePreset(const FObjectInitializer& ObjInit)
@@ -33,7 +34,13 @@ void UTakePreset::CreateLevelSequence()
 	LevelSequence = NewObject<ULevelSequence>(this, GetFName(), SequenceFlags);
 	LevelSequence->Initialize();
 
-	LevelSequence->GetMovieScene()->SetPlaybackRange(TRange<FFrameNumber>(0, TNumericLimits<int32>::Max()-1));
+	const UMovieSceneToolsProjectSettings* ProjectSettings = GetDefault<UMovieSceneToolsProjectSettings>();
+	UMovieScene* MovieScene = LevelSequence->GetMovieScene();
+
+	FFrameNumber StartFrame = (ProjectSettings->DefaultStartTime * MovieScene->GetTickResolution()).RoundToFrame();
+	int32        Duration = (ProjectSettings->DefaultDuration * MovieScene->GetTickResolution()).RoundToFrame().Value;
+
+	LevelSequence->GetMovieScene()->SetPlaybackRange(StartFrame, Duration);
 
 	FMovieSceneEditorData& EditorData = LevelSequence->GetMovieScene()->GetEditorData();
 	EditorData.ViewStart = -1.0;
