We still need to reverse out any skelmesh rotation that came from the root lock zeroing out the transform, but only if the animation is not attached to a socket, if it is it doesn't need that since it's in component space not in actor space.

diff --git a/Engine/Plugins/VirtualProduction/Takes/Source/TakeRecorderSources/Private/TakeRecorderActorSource.cpp b/Engine/Plugins/VirtualProduction/Takes/Source/TakeRecorderSources/Private/TakeRecorderActorSource.cpp
--- a/Engine/Plugins/VirtualProduction/Takes/Source/TakeRecorderSources/Private/TakeRecorderActorSource.cpp
+++ b/Engine/Plugins/VirtualProduction/Takes/Source/TakeRecorderSources/Private/TakeRecorderActorSource.cpp
@@ -1439,7 +1439,7 @@ bool UTakeRecorderActorSource::IsOtherActorBeingRecorded(AActor* OtherActor) con
 	{
 		if (UTakeRecorderActorSource* ActorSource = Cast<UTakeRecorderActorSource>(Source))
 		{
-			if (ActorSource->Target.Get() == OtherActor)
+			if (ActorSource->bEnabled && ActorSource->Target.Get() == OtherActor)
 			{
 				return true;
 			}
@@ -1469,6 +1469,30 @@ FGuid UTakeRecorderActorSource::GetRecordedActorGuid(class AActor* OtherActor) c
 	return FGuid();
 }
 
+FTransform UTakeRecorderActorSource::GetRecordedActorAnimationInitialRootTransform(class AActor* OtherActor) const
+{
+	UTakeRecorderSources* OwningSources = CastChecked<UTakeRecorderSources>(GetOuter());
+	for (UTakeRecorderSource* Source : OwningSources->GetSources())
+	{
+		if (UTakeRecorderActorSource* ActorSource = Cast<UTakeRecorderActorSource>(Source))
+		{
+			AActor* OtherTarget = ActorSource->Target.Get();
+			if (OtherTarget && OtherActor && (OtherTarget == OtherActor || OtherTarget->GetName() == OtherActor->GetName()))
+			{
+				for (UMovieSceneTrackRecorder* TrackRecorder : ActorSource->TrackRecorders)
+				{
+					if (TrackRecorder->IsA<UMovieSceneAnimationTrackRecorder>())
+					{
+						return Cast<UMovieSceneAnimationTrackRecorder>(TrackRecorder)->GetInitialRootTransform();
+					}
+				}
+
+			}
+		}
+	}
+	return FTransform::Identity;
+}
+
 FMovieSceneSequenceID UTakeRecorderActorSource::GetLevelSequenceID(class AActor* OtherActor)
 {
 	FMovieSceneSequenceID OutSequenceID = MovieSceneSequenceID::Root;
diff --git a/Engine/Plugins/VirtualProduction/Takes/Source/TakeRecorderSources/Private/TakeRecorderActorSource.h b/Engine/Plugins/VirtualProduction/Takes/Source/TakeRecorderSources/Private/TakeRecorderActorSource.h
--- a/Engine/Plugins/VirtualProduction/Takes/Source/TakeRecorderSources/Private/TakeRecorderActorSource.h
+++ b/Engine/Plugins/VirtualProduction/Takes/Source/TakeRecorderSources/Private/TakeRecorderActorSource.h
@@ -236,6 +236,8 @@ protected:
 	FMovieSceneSequenceID GetLevelSequenceID(class AActor* OtherActor) override;
 	/** Returns generic track recorder settings */
 	FTrackRecorderSettings GetTrackRecorderSettings() const override;
+	/** Returns offset that may get set when recording a skeletal mesh animation. Needed to correctly transform attached children*/
+	FTransform GetRecordedActorAnimationInitialRootTransform(class AActor* OtherActor) const override;
 
 	// ~IMovieSceneTrackRecorderHost Interface
 	
diff --git a/Engine/Plugins/VirtualProduction/Takes/Source/TakeTrackRecorders/Private/MovieScene3DTransformTrackRecorder.cpp b/Engine/Plugins/VirtualProduction/Takes/Source/TakeTrackRecorders/Private/MovieScene3DTransformTrackRecorder.cpp
--- a/Engine/Plugins/VirtualProduction/Takes/Source/TakeTrackRecorders/Private/MovieScene3DTransformTrackRecorder.cpp
+++ b/Engine/Plugins/VirtualProduction/Takes/Source/TakeTrackRecorders/Private/MovieScene3DTransformTrackRecorder.cpp
@@ -83,7 +83,7 @@ void UMovieScene3DTransformTrackRecorder::CreateTrackImpl()
 
 bool UMovieScene3DTransformTrackRecorder::ShouldAddNewKey(const FTransform& TransformToRecord)
 {
-	return !FTransform::AreTranslationsEqual(TransformToRecord, PreviousValue) || !FTransform::AreRotationsEqual(TransformToRecord, PreviousValue) || !FTransform::AreScale3DsEqual(TransformToRecord, PreviousValue);
+	return (bSetFirstKey || !FTransform::AreTranslationsEqual(TransformToRecord, PreviousValue) || !FTransform::AreRotationsEqual(TransformToRecord, PreviousValue) || !FTransform::AreScale3DsEqual(TransformToRecord, PreviousValue));
 }
 
 
@@ -112,6 +112,47 @@ void UMovieScene3DTransformTrackRecorder::FinalizeTrackImpl()
  	check (	BufferedTransforms.Times.Num() == BufferedTransforms.ScaleZ.Num());
  
  	SlowTask.EnterProgressFrame();
+ 
+	FTransform InvParentAnimationRootTransform = FTransform::Identity;
+	FName SocketName;
+	FName ComponentName;
+	AActor* ActorToRecord = Cast<AActor>(ObjectToRecord.Get());
+	if (ActorToRecord)
+	{
+		AActor* AttachedToActor = SequenceRecorderUtils::GetAttachment(ActorToRecord, SocketName, ComponentName);
+		//If attached to an actor that has an initial root transform(skeletal mesh) that we negated out when moving over the anim keys to the transform
+		//track we need to reapply that inverse here. But if it's attached to a socket we don't since we are in not in actor space but sill in component space.
+		if (AttachedToActor  && SocketName == NAME_None)
+		{
+			InvParentAnimationRootTransform = OwningTakeRecorderSource->GetRecordedActorAnimationInitialRootTransform(AttachedToActor).Inverse();
+			if (!InvParentAnimationRootTransform.Equals(FTransform::Identity))
+			{
+				if (DefaultTransform.IsSet())
+				{
+					FTransform DT = DefaultTransform.GetValue();
+					DT = DT * InvParentAnimationRootTransform;
+					DefaultTransform = DT;
+					FVector Translation = DT.GetTranslation();
+					FVector EulerRotation = DT.GetRotation().Rotator().Euler();
+					FVector Scale = DT.GetScale3D();
+					TArrayView<FMovieSceneFloatChannel*> FloatChannels = MovieSceneSection->GetChannelProxy().GetChannels<FMovieSceneFloatChannel>();
+					FloatChannels[0]->SetDefault(Translation.X);
+					FloatChannels[1]->SetDefault(Translation.Y);
+					FloatChannels[2]->SetDefault(Translation.Z);
+					FloatChannels[3]->SetDefault(EulerRotation.X);
+					FloatChannels[4]->SetDefault(EulerRotation.Y);
+					FloatChannels[5]->SetDefault(EulerRotation.Z);
+					FloatChannels[6]->SetDefault(Scale.X);
+					FloatChannels[7]->SetDefault(Scale.Y);
+					FloatChannels[8]->SetDefault(Scale.Z);
+				}
+				if (BufferedTransforms.Times.Num() > 0)
+				{
+					BufferedTransforms.PostMultTransform(InvParentAnimationRootTransform);
+				}
+			}
+		}
+	}
 
  	// Try to 're-wind' rotations that look like axis flips
  	// We need to do this as a post-process because the recorder cant reliably access 'wound' rotations:
@@ -313,10 +354,11 @@ void UMovieScene3DTransformTrackRecorder::RecordSampleImpl(const FQualifiedFrame
 				{
 					FSerializedTransform SerializedTransform(PreviousValue, CurrentFrame);
 					TransformSerializer.WriteFrameData(TransformSerializer.FramesWritten, SerializedTransform);
-					bSetFirstKey = false;
 				}
 				PreviousFrame = CurrentFrame;
 			}
+			bSetFirstKey = false;
+
 		}
  	}
 }
@@ -365,13 +407,13 @@ bool UMovieScene3DTransformTrackRecorder::ResolveTransformToRecord(FTransform& O
  			// We capture world space transforms for actors if they're attached, but we're not recording the attachment parent
 			bCaptureWorldSpaceTransform = !OwningTakeRecorderSource->IsOtherActorBeingRecorded(AttachParent->GetOwner());
  		}
- 
+
  		if (!RootComponent)
  		{
  			return false;
  		}
  
- 		if (bCaptureWorldSpaceTransform)
+	    if (bCaptureWorldSpaceTransform)
  		{
  			OutTransform = Actor->ActorToWorld();
  		}
diff --git a/Engine/Plugins/VirtualProduction/Takes/Source/TakeTrackRecorders/Public/TrackRecorders/IMovieSceneTrackRecorderHost.h b/Engine/Plugins/VirtualProduction/Takes/Source/TakeTrackRecorders/Public/TrackRecorders/IMovieSceneTrackRecorderHost.h
--- a/Engine/Plugins/VirtualProduction/Takes/Source/TakeTrackRecorders/Public/TrackRecorders/IMovieSceneTrackRecorderHost.h
+++ b/Engine/Plugins/VirtualProduction/Takes/Source/TakeTrackRecorders/Public/TrackRecorders/IMovieSceneTrackRecorderHost.h
@@ -168,5 +168,11 @@ public:
 	 */
 	virtual FTrackRecorderSettings GetTrackRecorderSettings() const = 0;
 
+	/** Returns original transform that may get set when recording a skeletal mesh animation. Needed to correctly transform attached children.
+	* @param OtherActor Actor to look for.
+	* @return The FTransform applied to the actor when recording animation, that the attached actor may need to apply to themselves to orient correctly.
+	*/
+	virtual FTransform GetRecordedActorAnimationInitialRootTransform(class AActor* OtherActor) const = 0;
+
 	/**/
 };
