Take Recorder: Sort take numbers

diff --git a/Engine/Plugins/VirtualProduction/Takes/Source/TakesCore/Private/TakeData.h b/Engine/Plugins/VirtualProduction/Takes/Source/TakesCore/Private/TakeData.h
--- a/Engine/Plugins/VirtualProduction/Takes/Source/TakesCore/Private/TakeData.h
+++ b/Engine/Plugins/VirtualProduction/Takes/Source/TakesCore/Private/TakeData.h
@@ -52,6 +52,8 @@ public:
 
 		CurrentTakeNumber = TakeMetaData->GetTakeNumber();
 
+		TakeNumbers.Sort();
+
 		return true;
 	}
 
