Take Recorder: Close sequencer editor on stopRecordingTake

diff --git a/Engine/Plugins/VirtualProduction/Takes/Source/TakeRecorderSources/Private/TakeRecorderSourcesModule.cpp b/Engine/Plugins/VirtualProduction/Takes/Source/TakeRecorderSources/Private/TakeRecorderSourcesModule.cpp
--- a/Engine/Plugins/VirtualProduction/Takes/Source/TakeRecorderSources/Private/TakeRecorderSourcesModule.cpp
+++ b/Engine/Plugins/VirtualProduction/Takes/Source/TakeRecorderSources/Private/TakeRecorderSourcesModule.cpp
@@ -31,6 +31,7 @@
 #include "Editor.h"
 #include "SceneOutlinerModule.h"
 #include "SceneOutlinerPublicTypes.h"
+#include "Toolkits/AssetEditorManager.h"
 
 #include "TakeRecorderMicrophoneAudioSource.h"
 #include "TakeRecorderWorldSource.h"
@@ -558,6 +559,12 @@ public:
 #if WITH_EDITOR
 		if (UTakeRecorder* ActiveRecorder = UTakeRecorder::GetActiveRecorder())
 		{
+			ULevelSequence* ActiveSequence = ActiveRecorder->GetSequence();
+			if (ActiveSequence)
+			{
+				FAssetEditorManager::Get().CloseAllEditorsForAsset(ActiveSequence);
+			}
+
 			ActiveRecorder->Stop();
 		}
 		return true;
