Change plugins that have a "developer" module to limit their scope to editor builds with -game/-server to "UncookedOnly". Developer modules are meant to be usable by runtime games, and cannot reference editor modules.

diff --git a/Engine/Plugins/Animation/LiveLink/LiveLink.uplugin b/Engine/Plugins/Animation/LiveLink/LiveLink.uplugin
--- a/Engine/Plugins/Animation/LiveLink/LiveLink.uplugin
+++ b/Engine/Plugins/Animation/LiveLink/LiveLink.uplugin
@@ -38,7 +38,7 @@
 		},
 		{
 			"Name": "LiveLinkGraphNode",
-			"Type": "Developer",
+			"Type": "UncookedOnly",
 			"LoadingPhase": "Default"
 		},
 		{
@@ -48,7 +48,7 @@
 		},
 		{
 			"Name": "LiveLinkSequencer",
-			"Type": "Developer",
+			"Type": "UncookedOnly",
 			"LoadingPhase": "Default"
 		}
 	]
diff --git a/Engine/Plugins/Developer/Concert/ConcertMain/ConcertMain.uplugin b/Engine/Plugins/Developer/Concert/ConcertMain/ConcertMain.uplugin
--- a/Engine/Plugins/Developer/Concert/ConcertMain/ConcertMain.uplugin
+++ b/Engine/Plugins/Developer/Concert/ConcertMain/ConcertMain.uplugin
@@ -19,13 +19,13 @@
 	"Modules": [
 		{
 			"Name": "Concert",
-			"Type": "Developer",
+			"Type": "UncookedOnly",
 			"LoadingPhase": "PreDefault",
 			"WhitelistPrograms": [ "UnrealMultiUserServer", "UnrealDisasterRecoveryService" ]
 		},
 		{
 			"Name": "ConcertTransport",
-			"Type": "Developer",
+			"Type": "UncookedOnly",
 			"LoadingPhase": "PreDefault",
 			"WhitelistPrograms": [ "UnrealMultiUserServer", "UnrealDisasterRecoveryService" ]
 		}
diff --git a/Engine/Plugins/Developer/Concert/ConcertSync/ConcertSyncClient/ConcertSyncClient.uplugin b/Engine/Plugins/Developer/Concert/ConcertSync/ConcertSyncClient/ConcertSyncClient.uplugin
--- a/Engine/Plugins/Developer/Concert/ConcertSync/ConcertSyncClient/ConcertSyncClient.uplugin
+++ b/Engine/Plugins/Developer/Concert/ConcertSync/ConcertSyncClient/ConcertSyncClient.uplugin
@@ -18,12 +18,12 @@
 	"Modules": [
 		{
 			"Name": "ConcertSyncClient",
-			"Type": "Developer",
+			"Type": "UncookedOnly",
 			"LoadingPhase": "Default"
 		},
 		{
 			"Name": "ConcertSyncUI",
-			"Type": "Developer",
+			"Type": "UncookedOnly",
 			"LoadingPhase": "Default"
 		}
 	],
diff --git a/Engine/Plugins/Developer/Concert/ConcertSync/ConcertSyncCore/ConcertSyncCore.uplugin b/Engine/Plugins/Developer/Concert/ConcertSync/ConcertSyncCore/ConcertSyncCore.uplugin
--- a/Engine/Plugins/Developer/Concert/ConcertSync/ConcertSyncCore/ConcertSyncCore.uplugin
+++ b/Engine/Plugins/Developer/Concert/ConcertSync/ConcertSyncCore/ConcertSyncCore.uplugin
@@ -19,7 +19,7 @@
 	"Modules": [
 		{
 			"Name": "ConcertSyncCore",
-			"Type": "Developer",
+			"Type": "UncookedOnly",
 			"LoadingPhase": "PreDefault",
 			"WhitelistPrograms": [ "UnrealMultiUserServer", "UnrealDisasterRecoveryService" ]
 		}
diff --git a/Engine/Plugins/Developer/Concert/ConcertSync/ConcertSyncServer/ConcertSyncServer.uplugin b/Engine/Plugins/Developer/Concert/ConcertSync/ConcertSyncServer/ConcertSyncServer.uplugin
--- a/Engine/Plugins/Developer/Concert/ConcertSync/ConcertSyncServer/ConcertSyncServer.uplugin
+++ b/Engine/Plugins/Developer/Concert/ConcertSync/ConcertSyncServer/ConcertSyncServer.uplugin
@@ -19,7 +19,7 @@
 	"Modules": [
 		{
 			"Name": "ConcertSyncServer",
-			"Type": "Developer",
+			"Type": "UncookedOnly",
 			"LoadingPhase": "Default",
 			"WhitelistPrograms": [ "UnrealMultiUserServer", "UnrealDisasterRecoveryService" ]
 		}
diff --git a/Engine/Plugins/Developer/Concert/ConcertSync/ConcertSyncTest/ConcertSyncTest.uplugin b/Engine/Plugins/Developer/Concert/ConcertSync/ConcertSyncTest/ConcertSyncTest.uplugin
--- a/Engine/Plugins/Developer/Concert/ConcertSync/ConcertSyncTest/ConcertSyncTest.uplugin
+++ b/Engine/Plugins/Developer/Concert/ConcertSync/ConcertSyncTest/ConcertSyncTest.uplugin
@@ -19,7 +19,7 @@
 	"Modules": [
 		{
 			"Name": "ConcertSyncTest",
-			"Type": "Developer",
+			"Type": "UncookedOnly",
 			"LoadingPhase": "PreDefault",
 			"WhitelistPrograms": []
 		}
diff --git a/Engine/Plugins/Developer/Concert/MultiUserClient/MultiUserClient.uplugin b/Engine/Plugins/Developer/Concert/MultiUserClient/MultiUserClient.uplugin
--- a/Engine/Plugins/Developer/Concert/MultiUserClient/MultiUserClient.uplugin
+++ b/Engine/Plugins/Developer/Concert/MultiUserClient/MultiUserClient.uplugin
@@ -17,7 +17,7 @@
 	"Modules": [
 		{
 			"Name": "MultiUserClient",
-			"Type": "Developer",
+			"Type": "UncookedOnly",
 			"LoadingPhase": "Default"
 		},
 		{
diff --git a/Engine/Plugins/Developer/OneSkyLocalizationService/OneSkyLocalizationService.uplugin b/Engine/Plugins/Developer/OneSkyLocalizationService/OneSkyLocalizationService.uplugin
--- a/Engine/Plugins/Developer/OneSkyLocalizationService/OneSkyLocalizationService.uplugin
+++ b/Engine/Plugins/Developer/OneSkyLocalizationService/OneSkyLocalizationService.uplugin
@@ -18,7 +18,7 @@
 	[
 		{
 			"Name" : "OneSkyLocalizationService",
-			"Type" : "Developer",
+			"Type" : "Editor",
 			"LoadingPhase" : "Default"
 		}
 	]
diff --git a/Engine/Plugins/Editor/AssetManagerEditor/AssetManagerEditor.uplugin b/Engine/Plugins/Editor/AssetManagerEditor/AssetManagerEditor.uplugin
--- a/Engine/Plugins/Editor/AssetManagerEditor/AssetManagerEditor.uplugin
+++ b/Engine/Plugins/Editor/AssetManagerEditor/AssetManagerEditor.uplugin
@@ -19,7 +19,7 @@
 	[
 		{
 			"Name" : "AssetManagerEditor",
-			"Type": "Developer",
+			"Type": "UncookedOnly" ,
 			"LoadingPhase" : "PreDefault"
 		}
 	]
diff --git a/Engine/Plugins/Editor/BlueprintMaterialTextureNodes/BlueprintMaterialTextureNodes.uplugin b/Engine/Plugins/Editor/BlueprintMaterialTextureNodes/BlueprintMaterialTextureNodes.uplugin
--- a/Engine/Plugins/Editor/BlueprintMaterialTextureNodes/BlueprintMaterialTextureNodes.uplugin
+++ b/Engine/Plugins/Editor/BlueprintMaterialTextureNodes/BlueprintMaterialTextureNodes.uplugin
@@ -16,7 +16,7 @@
 	"Modules": [
 		{
 			"Name": "BlueprintMaterialTextureNodes",
-			"Type": "Developer",
+			"Type": "UncookedOnly",
 			"LoadingPhase": "PreLoadingScreen"
 		}
 	]
diff --git a/Engine/Plugins/Editor/CryptoKeys/CryptoKeys.uplugin b/Engine/Plugins/Editor/CryptoKeys/CryptoKeys.uplugin
--- a/Engine/Plugins/Editor/CryptoKeys/CryptoKeys.uplugin
+++ b/Engine/Plugins/Editor/CryptoKeys/CryptoKeys.uplugin
@@ -17,7 +17,7 @@
 	"Modules": [
 		{
 			"Name": "CryptoKeys",
-			"Type": "Developer",
+			"Type": "Editor",
 			"LoadingPhase": "Default",
 			"WhitelistPlatforms": [
 				"Mac",
@@ -27,7 +27,7 @@
 		},
 		{
 			"Name": "CryptoKeysOpenSSL",
-			"Type": "Developer",
+			"Type": "Editor",
 			"LoadingPhase": "Default",
 			"WhitelistPlatforms": [
 				"Mac",
diff --git a/Engine/Plugins/Editor/GameplayTagsEditor/GameplayTagsEditor.uplugin b/Engine/Plugins/Editor/GameplayTagsEditor/GameplayTagsEditor.uplugin
--- a/Engine/Plugins/Editor/GameplayTagsEditor/GameplayTagsEditor.uplugin
+++ b/Engine/Plugins/Editor/GameplayTagsEditor/GameplayTagsEditor.uplugin
@@ -18,7 +18,7 @@
 	[
 		{
 			"Name" : "GameplayTagsEditor",
-			"Type" : "Developer",
+			"Type" : "UncookedOnly",
 			"LoadingPhase" : "PreDefault"
 		}
 	]
diff --git a/Engine/Plugins/Experimental/AppleVision/AppleVision.uplugin b/Engine/Plugins/Experimental/AppleVision/AppleVision.uplugin
--- a/Engine/Plugins/Experimental/AppleVision/AppleVision.uplugin
+++ b/Engine/Plugins/Experimental/AppleVision/AppleVision.uplugin
@@ -27,7 +27,7 @@
 		},
 		{
 			"Name" : "AppleVisionBlueprintSupport",
-			"Type" : "Developer",
+			"Type" : "UncookedOnly",
 			"LoadingPhase" : "Default",
 			"WhitelistPlatforms" :
 			[
diff --git a/Engine/Plugins/Experimental/ControlRig/ControlRig.uplugin b/Engine/Plugins/Experimental/ControlRig/ControlRig.uplugin
--- a/Engine/Plugins/Experimental/ControlRig/ControlRig.uplugin
+++ b/Engine/Plugins/Experimental/ControlRig/ControlRig.uplugin
@@ -23,7 +23,7 @@
 		},
 		{
 			"Name" : "ControlRigDeveloper",
-			"Type" : "Developer",
+			"Type" : "UncookedOnly",
 			"LoadingPhase" : "PreDefault"
 		},
 		{
diff --git a/Engine/Plugins/Experimental/PanoramicCapture/PanoramicCapture.uplugin b/Engine/Plugins/Experimental/PanoramicCapture/PanoramicCapture.uplugin
--- a/Engine/Plugins/Experimental/PanoramicCapture/PanoramicCapture.uplugin
+++ b/Engine/Plugins/Experimental/PanoramicCapture/PanoramicCapture.uplugin
@@ -17,7 +17,7 @@
 	"Modules": [
 		{
 			"Name": "PanoramicCapture",
-			"Type": "Developer",
+			"Type": "UncookedOnly",
 			"LoadingPhase": "Default",
 			"WhitelistPlatforms": [
 				"Win64"
diff --git a/Engine/Plugins/Experimental/VirtualCamera/VirtualCamera.uplugin b/Engine/Plugins/Experimental/VirtualCamera/VirtualCamera.uplugin
--- a/Engine/Plugins/Experimental/VirtualCamera/VirtualCamera.uplugin
+++ b/Engine/Plugins/Experimental/VirtualCamera/VirtualCamera.uplugin
@@ -17,7 +17,7 @@
   "Modules": [
     {
       "Name": "VirtualCamera",
-      "Type": "Runtime",
+      "Type": "UncookedOnly",
       "LoadingPhase": "PostDefault",
       "WhitelistPlatforms" :
 			[
diff --git a/Engine/Plugins/MovieScene/SequencerScripting/SequencerScripting.uplugin b/Engine/Plugins/MovieScene/SequencerScripting/SequencerScripting.uplugin
--- a/Engine/Plugins/MovieScene/SequencerScripting/SequencerScripting.uplugin
+++ b/Engine/Plugins/MovieScene/SequencerScripting/SequencerScripting.uplugin
@@ -22,7 +22,7 @@
 	"Modules": [
 		{
 			"Name": "SequencerScripting",
-			"Type": "Developer",
+			"Type": "UncookedOnly",
 			"LoadingPhase": "Default"
 		},
 		{
diff --git a/Engine/Plugins/Online/OnlineSubsystemUtils/OnlineSubsystemUtils.uplugin b/Engine/Plugins/Online/OnlineSubsystemUtils/OnlineSubsystemUtils.uplugin
--- a/Engine/Plugins/Online/OnlineSubsystemUtils/OnlineSubsystemUtils.uplugin
+++ b/Engine/Plugins/Online/OnlineSubsystemUtils/OnlineSubsystemUtils.uplugin
@@ -18,7 +18,7 @@
         },
         {
             "Name": "OnlineBlueprintSupport",
-            "Type": "Developer",
+            "Type": "UncookedOnly",
             "LoadingPhase": "PostConfigInit"
         }
     ],
diff --git a/Engine/Plugins/Runtime/ApexDestruction/ApexDestruction.uplugin b/Engine/Plugins/Runtime/ApexDestruction/ApexDestruction.uplugin
--- a/Engine/Plugins/Runtime/ApexDestruction/ApexDestruction.uplugin
+++ b/Engine/Plugins/Runtime/ApexDestruction/ApexDestruction.uplugin
@@ -26,7 +26,7 @@
 		},
     {
 			"Name" : "ApexDestructionEditor",
-			"Type" : "Developer",
+			"Type" : "UncookedOnly",
 			"LoadingPhase" : "PreDefault"
 		}
 	]
diff --git a/Engine/Plugins/Runtime/AppleImageUtils/AppleImageUtils.uplugin b/Engine/Plugins/Runtime/AppleImageUtils/AppleImageUtils.uplugin
--- a/Engine/Plugins/Runtime/AppleImageUtils/AppleImageUtils.uplugin
+++ b/Engine/Plugins/Runtime/AppleImageUtils/AppleImageUtils.uplugin
@@ -23,7 +23,7 @@
 		},
 		{
 			"Name" : "AppleImageUtilsBlueprintSupport",
-			"Type" : "Developer",
+			"Type" : "UncookedOnly",
 			"LoadingPhase" : "Default",
 			"WhitelistPlatforms" :
 			[
diff --git a/Engine/Plugins/Runtime/GameplayAbilities/GameplayAbilities.uplugin b/Engine/Plugins/Runtime/GameplayAbilities/GameplayAbilities.uplugin
--- a/Engine/Plugins/Runtime/GameplayAbilities/GameplayAbilities.uplugin
+++ b/Engine/Plugins/Runtime/GameplayAbilities/GameplayAbilities.uplugin
@@ -23,7 +23,7 @@
 		},
 		{
 			"Name" : "GameplayAbilitiesEditor",
-			"Type" : "Developer",
+			"Type" : "UncookedOnly",
 			"LoadingPhase" : "PreDefault"
 		}
 	],
diff --git a/Engine/Plugins/Runtime/PhysXVehicles/PhysXVehicles.uplugin b/Engine/Plugins/Runtime/PhysXVehicles/PhysXVehicles.uplugin
--- a/Engine/Plugins/Runtime/PhysXVehicles/PhysXVehicles.uplugin
+++ b/Engine/Plugins/Runtime/PhysXVehicles/PhysXVehicles.uplugin
@@ -23,7 +23,7 @@
 		},
 		{
 			"Name" : "PhysXVehiclesEditor",
-			"Type" : "Developer",
+			"Type" : "UncookedOnly",
 			"LoadingPhase" : "PreDefault"
 		}
 	]
diff --git a/Engine/Plugins/VirtualProduction/Takes/Takes.uplugin b/Engine/Plugins/VirtualProduction/Takes/Takes.uplugin
--- a/Engine/Plugins/VirtualProduction/Takes/Takes.uplugin
+++ b/Engine/Plugins/VirtualProduction/Takes/Takes.uplugin
@@ -24,22 +24,22 @@
 	[
 		{
 			"Name" : "TakesCore",
-			"Type" : "Developer",
+			"Type" : "UncookedOnly",
 			"LoadingPhase" : "PostEngineInit"
 		},
 		{
 			"Name" : "TakeTrackRecorders",
-			"Type" : "Developer",
+			"Type" : "UncookedOnly",
 			"LoadingPhase" : "PostEngineInit"
 		},
 		{
 			"Name" : "TakeRecorderSources",
-			"Type" : "Developer",
+			"Type" : "UncookedOnly",
 			"LoadingPhase" : "PostEngineInit"
 		},
 		{
 			"Name" : "TakeRecorder",
-			"Type" : "Developer",
+			"Type" : "UncookedOnly",
 			"LoadingPhase" : "PostEngineInit"
 		},
 	]
