Fix crash running monolithic editor due to global constructor trying to construct a UObject before the UObject subsystem has been initialized.

diff --git a/Engine/Plugins/VirtualProduction/Takes/Source/TakeRecorder/Private/Recorder/TakeRecorder.cpp b/Engine/Plugins/VirtualProduction/Takes/Source/TakeRecorder/Private/Recorder/TakeRecorder.cpp
--- a/Engine/Plugins/VirtualProduction/Takes/Source/TakeRecorder/Private/Recorder/TakeRecorder.cpp
+++ b/Engine/Plugins/VirtualProduction/Takes/Source/TakeRecorder/Private/Recorder/TakeRecorder.cpp
@@ -330,13 +330,17 @@ public:
 FTickableTakeRecorder TickableTakeRecorder;
 
 // Static members of UTakeRecorder
-TStrongObjectPtr<UTakeRecorder> CurrentRecorder;
+static TStrongObjectPtr<UTakeRecorder>& GetCurrentRecorder()
+{
+	static TStrongObjectPtr<UTakeRecorder> CurrentRecorder;
+	return CurrentRecorder;
+}
 FOnTakeRecordingInitialized UTakeRecorder::OnRecordingInitializedEvent;
 
 // Static functions for UTakeRecorder
 UTakeRecorder* UTakeRecorder::GetActiveRecorder()
 {
-	return CurrentRecorder.Get();
+	return GetCurrentRecorder().Get();
 }
 
 FOnTakeRecordingInitialized& UTakeRecorder::OnRecordingInitialized()
@@ -346,13 +350,13 @@ FOnTakeRecordingInitialized& UTakeRecorder::OnRecordingInitialized()
 
 bool UTakeRecorder::SetActiveRecorder(UTakeRecorder* NewActiveRecorder)
 {
-	if (CurrentRecorder.IsValid())
+	if (GetCurrentRecorder().IsValid())
 	{
 		return false;
 	}
 
-	CurrentRecorder.Reset(NewActiveRecorder);
-	TickableTakeRecorder.WeakRecorder = CurrentRecorder.Get();
+	GetCurrentRecorder().Reset(NewActiveRecorder);
+	TickableTakeRecorder.WeakRecorder = GetCurrentRecorder().Get();
 	OnRecordingInitializedEvent.Broadcast(NewActiveRecorder);
 	return true;
 }
@@ -826,9 +830,9 @@ void UTakeRecorder::Stop()
 	OnStopCleanup.Reset();
 
 	// reset the current recorder and stop us from being ticked
-	if (CurrentRecorder.Get() == this)
+	if (GetCurrentRecorder().Get() == this)
 	{
-		CurrentRecorder.Reset();
+		GetCurrentRecorder().Reset();
 		TickableTakeRecorder.WeakRecorder = nullptr;
 
 		if (bDidEverStartRecording)
