Take Recorder: Added reentrant guards to prevent starting and stopping recordings while they're in progress

diff --git a/Engine/Plugins/VirtualProduction/Takes/Source/TakeRecorder/Private/Recorder/TakeRecorder.cpp b/Engine/Plugins/VirtualProduction/Takes/Source/TakeRecorder/Private/Recorder/TakeRecorder.cpp
--- a/Engine/Plugins/VirtualProduction/Takes/Source/TakeRecorder/Private/Recorder/TakeRecorder.cpp
+++ b/Engine/Plugins/VirtualProduction/Takes/Source/TakeRecorder/Private/Recorder/TakeRecorder.cpp
@@ -761,6 +761,15 @@ void UTakeRecorder::Start(const FTimecode& InTimecodeSource)
 
 void UTakeRecorder::Stop()
 {
+	static bool bStoppedRecording = false;
+
+	if (bStoppedRecording)
+	{
+		return;
+	}
+
+	TGuardValue<bool> ReentrantGuard(bStoppedRecording, true);
+
 	FTimecode TimecodeOut = FApp::GetTimecode();
 	USequencerSettings* SequencerSettings = USequencerSettingsContainer::GetOrCreate<USequencerSettings>(TEXT("TakeRecorderSequenceEditor"));
 
diff --git a/Engine/Plugins/VirtualProduction/Takes/Source/TakeRecorder/Private/Widgets/STakeRecorderCockpit.cpp b/Engine/Plugins/VirtualProduction/Takes/Source/TakeRecorder/Private/Widgets/STakeRecorderCockpit.cpp
--- a/Engine/Plugins/VirtualProduction/Takes/Source/TakeRecorder/Private/Widgets/STakeRecorderCockpit.cpp
+++ b/Engine/Plugins/VirtualProduction/Takes/Source/TakeRecorder/Private/Widgets/STakeRecorderCockpit.cpp
@@ -875,6 +875,15 @@ void STakeRecorderCockpit::StopRecording()
 
 void STakeRecorderCockpit::StartRecording()
 {
+	static bool bStartedRecording = false;
+
+	if (bStartedRecording)
+	{
+		return;
+	}
+
+	TGuardValue<bool> ReentrantGuard(bStartedRecording, true);
+
 	ULevelSequence*       LevelSequence = LevelSequenceAttribute.Get();
 	UTakeRecorderSources* Sources = LevelSequence ? LevelSequence->FindMetaData<UTakeRecorderSources>() : nullptr;
 	
