Take Recorder: Close the recorded sequence on end pie if it wasn't closed otherwise.

diff --git a/Engine/Plugins/VirtualProduction/Takes/Source/TakeRecorder/Private/Recorder/TakeRecorder.cpp b/Engine/Plugins/VirtualProduction/Takes/Source/TakeRecorder/Private/Recorder/TakeRecorder.cpp
--- a/Engine/Plugins/VirtualProduction/Takes/Source/TakeRecorder/Private/Recorder/TakeRecorder.cpp
+++ b/Engine/Plugins/VirtualProduction/Takes/Source/TakeRecorder/Private/Recorder/TakeRecorder.cpp
@@ -920,6 +920,11 @@ void UTakeRecorder::HandleEndPIE(bool bIsSimulating)
 	PendingNotification->SetCompletionState(SNotificationItem::CS_Success);
 	
 	Stop();
+
+	if (FinishedAsset)
+	{
+		GEditor->GetEditorSubsystem<UAssetEditorSubsystem>()->CloseAllEditorsForAsset(FinishedAsset);
+	}
 }
 
 #undef LOCTEXT_NAMESPACE
\ No newline at end of file
