[VCam] Fetch the frame rate from the Recorder cockpit.

diff --git a/Engine/Plugins/VirtualProduction/Takes/Source/TakeRecorder/Private/Recorder/TakeRecorderPanel.cpp b/Engine/Plugins/VirtualProduction/Takes/Source/TakeRecorder/Private/Recorder/TakeRecorderPanel.cpp
--- a/Engine/Plugins/VirtualProduction/Takes/Source/TakeRecorder/Private/Recorder/TakeRecorderPanel.cpp
+++ b/Engine/Plugins/VirtualProduction/Takes/Source/TakeRecorder/Private/Recorder/TakeRecorderPanel.cpp
@@ -102,6 +102,16 @@ UTakeMetaData* UTakeRecorderPanel::GetTakeMetaData() const
 	return WeakTabContent.Pin()->GetTakeMetaData();
 }
 
+FFrameRate UTakeRecorderPanel::GetFrameRate() const
+{
+	if (!ValidateTabContent())
+	{
+		return FFrameRate();
+	}
+
+	return WeakTabContent.Pin()->GetFrameRate();
+}
+
 UTakeRecorderSources* UTakeRecorderPanel::GetSources() const
 {
 	if (!ValidateTabContent())
diff --git a/Engine/Plugins/VirtualProduction/Takes/Source/TakeRecorder/Private/Widgets/STakeRecorderCockpit.cpp b/Engine/Plugins/VirtualProduction/Takes/Source/TakeRecorder/Private/Widgets/STakeRecorderCockpit.cpp
--- a/Engine/Plugins/VirtualProduction/Takes/Source/TakeRecorder/Private/Widgets/STakeRecorderCockpit.cpp
+++ b/Engine/Plugins/VirtualProduction/Takes/Source/TakeRecorder/Private/Widgets/STakeRecorderCockpit.cpp
@@ -653,19 +653,24 @@ FReply STakeRecorderCockpit::SetFrameRate()
 	return FReply::Unhandled();
 }
 
-FText STakeRecorderCockpit::GetFrameRateText() const
+FFrameRate STakeRecorderCockpit::GetFrameRate() const
 {
 	// If not recorded, return app timecode frame rate
 	if (TakeMetaData->GetTimestamp() == FDateTime(0))
 	{
-		return FApp::GetTimecodeFrameRate().ToPrettyText();
+		return FApp::GetTimecodeFrameRate();
 	}
 	else
 	{
-		return TakeMetaData->GetFrameRate().ToPrettyText();
+		return TakeMetaData->GetFrameRate();
 	}
 }
 
+FText STakeRecorderCockpit::GetFrameRateText() const
+{
+	return GetFrameRate().ToPrettyText();
+}
+
 FText STakeRecorderCockpit::GetFrameRateTooltipText() const
 {
 	// If not recorded, return app timecode frame rate
diff --git a/Engine/Plugins/VirtualProduction/Takes/Source/TakeRecorder/Private/Widgets/STakeRecorderCockpit.h b/Engine/Plugins/VirtualProduction/Takes/Source/TakeRecorder/Private/Widgets/STakeRecorderCockpit.h
--- a/Engine/Plugins/VirtualProduction/Takes/Source/TakeRecorder/Private/Widgets/STakeRecorderCockpit.h
+++ b/Engine/Plugins/VirtualProduction/Takes/Source/TakeRecorder/Private/Widgets/STakeRecorderCockpit.h
@@ -39,6 +39,8 @@ public:
 		return TakeMetaData;
 	}
 
+	FFrameRate GetFrameRate() const;
+
 	bool Reviewing() const;
 
 	bool Recording() const;
diff --git a/Engine/Plugins/VirtualProduction/Takes/Source/TakeRecorder/Private/Widgets/STakeRecorderTabContent.cpp b/Engine/Plugins/VirtualProduction/Takes/Source/TakeRecorder/Private/Widgets/STakeRecorderTabContent.cpp
--- a/Engine/Plugins/VirtualProduction/Takes/Source/TakeRecorder/Private/Widgets/STakeRecorderTabContent.cpp
+++ b/Engine/Plugins/VirtualProduction/Takes/Source/TakeRecorder/Private/Widgets/STakeRecorderTabContent.cpp
@@ -197,6 +197,14 @@ UTakeMetaData* STakeRecorderTabContent::GetTakeMetaData() const
 	}
 }
 
+FFrameRate STakeRecorderTabContent::GetFrameRate() const
+{
+	TSharedPtr<STakeRecorderPanel>   Panel = WeakPanel.Pin();
+	TSharedPtr<STakeRecorderCockpit> Cockpit = Panel.IsValid() ? Panel->GetCockpitWidget() : nullptr;
+
+	return Cockpit.IsValid() ? Cockpit->GetFrameRate() : FFrameRate();
+}
+
 UTakeRecorderSources* STakeRecorderTabContent::GetSources() const
 {
 	ULevelSequence* LevelSequence = nullptr;
diff --git a/Engine/Plugins/VirtualProduction/Takes/Source/TakeRecorder/Private/Widgets/STakeRecorderTabContent.h b/Engine/Plugins/VirtualProduction/Takes/Source/TakeRecorder/Private/Widgets/STakeRecorderTabContent.h
--- a/Engine/Plugins/VirtualProduction/Takes/Source/TakeRecorder/Private/Widgets/STakeRecorderTabContent.h
+++ b/Engine/Plugins/VirtualProduction/Takes/Source/TakeRecorder/Private/Widgets/STakeRecorderTabContent.h
@@ -43,6 +43,8 @@ public:
 
 	UTakeMetaData* GetTakeMetaData() const;
 
+	FFrameRate GetFrameRate() const;
+
 	UTakeRecorderSources* GetSources() const;
 
 	void StartRecording() const;
diff --git a/Engine/Plugins/VirtualProduction/Takes/Source/TakeRecorder/Public/Recorder/TakeRecorderPanel.h b/Engine/Plugins/VirtualProduction/Takes/Source/TakeRecorder/Public/Recorder/TakeRecorderPanel.h
--- a/Engine/Plugins/VirtualProduction/Takes/Source/TakeRecorder/Public/Recorder/TakeRecorderPanel.h
+++ b/Engine/Plugins/VirtualProduction/Takes/Source/TakeRecorder/Public/Recorder/TakeRecorderPanel.h
@@ -84,6 +84,13 @@ public:
 	UTakeMetaData* GetTakeMetaData() const;
 
 
+	/**
+	 * Access the frame rate for this take
+	 */
+	UFUNCTION(BlueprintCallable, Category="Take Recorder|Panel")
+	FFrameRate GetFrameRate() const;
+
+
 	/**
 	 * Access the sources that are to be (or were) used for recording this take
 	 */
