Fix CIS build. Missing include.

diff --git a/Engine/Plugins/VirtualProduction/Takes/Source/TakeRecorder/Public/Recorder/TakeRecorderPanel.h b/Engine/Plugins/VirtualProduction/Takes/Source/TakeRecorder/Public/Recorder/TakeRecorderPanel.h
--- a/Engine/Plugins/VirtualProduction/Takes/Source/TakeRecorder/Public/Recorder/TakeRecorderPanel.h
+++ b/Engine/Plugins/VirtualProduction/Takes/Source/TakeRecorder/Public/Recorder/TakeRecorderPanel.h
@@ -2,6 +2,7 @@
 
 #pragma once
 
+#include "Misc/FrameRate.h"
 #include "TakeRecorderPanel.generated.h"
 
 class UTakePreset;
