Take Recorder: Disable transactional on movie scene during creation of take so that creation of the asset and operations on the movie scene before stopping recording are not undoable.

diff --git a/Engine/Plugins/VirtualProduction/Takes/Source/TakesCore/Private/TakeRecorderSources.cpp b/Engine/Plugins/VirtualProduction/Takes/Source/TakesCore/Private/TakeRecorderSources.cpp
--- a/Engine/Plugins/VirtualProduction/Takes/Source/TakesCore/Private/TakeRecorderSources.cpp
+++ b/Engine/Plugins/VirtualProduction/Takes/Source/TakesCore/Private/TakeRecorderSources.cpp
@@ -624,6 +624,9 @@ void UTakeRecorderSources::StopRecording(class ULevelSequence* InSequence, FTake
 		++SourcesSerialNumber;
 	}
 
+	// Re-enable transactional after recording
+	InSequence->GetMovieScene()->SetFlags(RF_Transactional);
+
 	// Ensure each sub-section is as long as it should be. If we're recording into subsections and a user is doing a partial
 	// re-record of the data within the sub section we can end up with the case where the new section is shorter than the original
 	// data. We don't want to trim the data unnecessarily, and we've been updating the length of the section every frame of the recording
@@ -646,6 +649,9 @@ void UTakeRecorderSources::StopRecording(class ULevelSequence* InSequence, FTake
 			AssetMetaData->Lock();
 
 			SubSection->SetRange(SubSequence->GetMovieScene()->GetPlaybackRange());
+
+			// Re-enable transactional after recording
+			SubSequence->GetMovieScene()->SetFlags(RF_Transactional);
 		}
 	}
 
@@ -722,6 +728,9 @@ ULevelSequence* UTakeRecorderSources::CreateSubSequenceForSource(ULevelSequence*
 		// past it will have the right resolution, but if the user modified it then we will preserve their desired resolution.
 		if (!ExistingSubSequence)
 		{
+			// Movie scene should not be transactional during the recording process
+			OutAsset->GetMovieScene()->ClearFlags(RF_Transactional);
+
 			OutAsset->GetMovieScene()->SetTickResolutionDirectly(InMasterSequence->GetMovieScene()->GetTickResolution());
 			OutAsset->GetMovieScene()->SetDisplayRate(InMasterSequence->GetMovieScene()->GetDisplayRate());
 		}
