Take Recorder: Added PreInitialize which is called before take recorder does any initialization for recording (ie. when the button is pressed)

diff --git a/Engine/Plugins/VirtualProduction/Takes/Source/TakeRecorder/Private/Recorder/TakeRecorder.cpp b/Engine/Plugins/VirtualProduction/Takes/Source/TakeRecorder/Private/Recorder/TakeRecorder.cpp
--- a/Engine/Plugins/VirtualProduction/Takes/Source/TakeRecorder/Private/Recorder/TakeRecorder.cpp
+++ b/Engine/Plugins/VirtualProduction/Takes/Source/TakeRecorder/Private/Recorder/TakeRecorder.cpp
@@ -389,6 +389,10 @@ bool UTakeRecorder::Initialize(ULevelSequence* LevelSequenceBase, UTakeRecorderS
 		return false;
 	}
 
+	OnRecordingPreInitializeEvent.Broadcast(this);
+
+	UTakeRecorderBlueprintLibrary::OnTakeRecorderPreInitialize();
+
 	if (!CreateDestinationAsset(*InParameters.Project.GetTakeAssetPath(), LevelSequenceBase, Sources, MetaData, OutError))
 	{
 		return false;
@@ -846,6 +850,11 @@ void UTakeRecorder::Stop()
 	RemoveFromRoot();
 }
 
+FOnTakeRecordingPreInitialize& UTakeRecorder::OnRecordingPreInitialize()
+{
+	return OnRecordingPreInitializeEvent;
+}
+
 FOnTakeRecordingStarted& UTakeRecorder::OnRecordingStarted()
 {
 	return OnRecordingStartedEvent;
diff --git a/Engine/Plugins/VirtualProduction/Takes/Source/TakeRecorder/Private/Recorder/TakeRecorderBlueprintLibrary.cpp b/Engine/Plugins/VirtualProduction/Takes/Source/TakeRecorder/Private/Recorder/TakeRecorderBlueprintLibrary.cpp
--- a/Engine/Plugins/VirtualProduction/Takes/Source/TakeRecorder/Private/Recorder/TakeRecorderBlueprintLibrary.cpp
+++ b/Engine/Plugins/VirtualProduction/Takes/Source/TakeRecorder/Private/Recorder/TakeRecorderBlueprintLibrary.cpp
@@ -15,6 +15,7 @@ namespace
 {
 	static UTakeRecorderPanel* CurrentTakeRecorderPanel;
 	static UTakeRecorderBlueprintLibrary::FOnTakeRecorderPanelChanged TakeRecorderPanelChanged;
+	static UTakeRecorderBlueprintLibrary::FOnTakeRecorderPreInitialize TakeRecorderPreInitialize;
 	static UTakeRecorderBlueprintLibrary::FOnTakeRecorderStarted TakeRecorderStarted;
 	static UTakeRecorderBlueprintLibrary::FOnTakeRecorderStopped TakeRecorderStopped;
 	static UTakeRecorderBlueprintLibrary::FOnTakeRecorderFinished TakeRecorderFinished;
@@ -131,6 +132,11 @@ void UTakeRecorderBlueprintLibrary::SetOnTakeRecorderPanelChanged(FOnTakeRecorde
 	TakeRecorderPanelChanged = OnTakeRecorderPanelChanged;
 };
 
+void UTakeRecorderBlueprintLibrary::SetOnTakeRecorderPreInitialize(FOnTakeRecorderPreInitialize OnTakeRecorderPreInitialize)
+{
+	TakeRecorderPreInitialize = OnTakeRecorderPreInitialize;
+};
+
 void UTakeRecorderBlueprintLibrary::SetOnTakeRecorderStarted(FOnTakeRecorderStarted OnTakeRecorderStarted)
 {
 	TakeRecorderStarted = OnTakeRecorderStarted;
@@ -156,6 +162,11 @@ void UTakeRecorderBlueprintLibrary::SetOnTakeRecorderMarkedFrameAdded(FOnTakeRec
 	TakeRecorderMarkedFrameAdded = OnTakeRecorderMarkedFrameAdded;
 };
 
+void UTakeRecorderBlueprintLibrary::OnTakeRecorderPreInitialize()
+{
+	TakeRecorderPreInitialize.ExecuteIfBound();
+}
+
 void UTakeRecorderBlueprintLibrary::OnTakeRecorderStarted()
 {
 	TakeRecorderStarted.ExecuteIfBound();
diff --git a/Engine/Plugins/VirtualProduction/Takes/Source/TakeRecorder/Public/Recorder/TakeRecorder.h b/Engine/Plugins/VirtualProduction/Takes/Source/TakeRecorder/Public/Recorder/TakeRecorder.h
--- a/Engine/Plugins/VirtualProduction/Takes/Source/TakeRecorder/Public/Recorder/TakeRecorder.h
+++ b/Engine/Plugins/VirtualProduction/Takes/Source/TakeRecorder/Public/Recorder/TakeRecorder.h
@@ -31,6 +31,7 @@ enum class ETakeRecorderState : uint8
 	Cancelled,
 };
 
+DECLARE_MULTICAST_DELEGATE_OneParam(FOnTakeRecordingPreInitialize, UTakeRecorder*);
 DECLARE_MULTICAST_DELEGATE_OneParam(FOnTakeRecordingInitialized, UTakeRecorder*);
 DECLARE_MULTICAST_DELEGATE_OneParam(FOnTakeRecordingStarted, UTakeRecorder*);
 DECLARE_MULTICAST_DELEGATE_OneParam(FOnTakeRecordingFinished, UTakeRecorder*);
@@ -109,18 +110,23 @@ public:
 	void Stop();
 
 	/**
-	* Retrieve a multi-cast delegate that is triggered when this recording starts
-	*/
+	 * Retrieve a multi-cast delegate that is triggered before initialization occurs (ie. when the recording button is pressed and before the countdown starts)
+	 */
+	FOnTakeRecordingPreInitialize& OnRecordingPreInitialize();
+
+	/**
+	 * Retrieve a multi-cast delegate that is triggered when this recording starts
+	 */
 	FOnTakeRecordingStarted& OnRecordingStarted();
 
 	/**
-	* Retrieve a multi-cast delegate that is triggered when this recording finishes
-	*/
+	 * Retrieve a multi-cast delegate that is triggered when this recording finishes
+	 */
 	FOnTakeRecordingFinished& OnRecordingFinished();
 
 	/**
-	* Retrieve a multi-cast delegate that is triggered when this recording is cancelled
-	*/
+	 * Retrieve a multi-cast delegate that is triggered when this recording is cancelled
+	 */
 	FOnTakeRecordingCancelled& OnRecordingCancelled();
 
 private:
@@ -202,6 +208,9 @@ private:
 	/** Anonymous array of cleanup functions to perform when a recording has finished */
 	TArray<TFunction<void()>> OnStopCleanup;
 
+	/** Triggered before the recorder is initialized */
+	FOnTakeRecordingPreInitialize OnRecordingPreInitializeEvent;
+
 	/** Triggered when this recorder starts */
 	FOnTakeRecordingStarted OnRecordingStartedEvent;
 
diff --git a/Engine/Plugins/VirtualProduction/Takes/Source/TakeRecorder/Public/Recorder/TakeRecorderBlueprintLibrary.h b/Engine/Plugins/VirtualProduction/Takes/Source/TakeRecorder/Public/Recorder/TakeRecorderBlueprintLibrary.h
--- a/Engine/Plugins/VirtualProduction/Takes/Source/TakeRecorder/Public/Recorder/TakeRecorderBlueprintLibrary.h
+++ b/Engine/Plugins/VirtualProduction/Takes/Source/TakeRecorder/Public/Recorder/TakeRecorderBlueprintLibrary.h
@@ -81,6 +81,7 @@ public:
 
 
 	DECLARE_DYNAMIC_DELEGATE(FOnTakeRecorderPanelChanged);
+	DECLARE_DYNAMIC_DELEGATE(FOnTakeRecorderPreInitialize);
 	DECLARE_DYNAMIC_DELEGATE(FOnTakeRecorderStarted);
 	DECLARE_DYNAMIC_DELEGATE(FOnTakeRecorderStopped);
 	DECLARE_DYNAMIC_DELEGATE_OneParam(FOnTakeRecorderFinished, ULevelSequence*, SequenceAsset);
@@ -91,6 +92,10 @@ public:
 	UFUNCTION(BlueprintCallable, Category="Take Recorder", meta=(DisplayName="Set On Take Recorder Panel Changed"))
 	static void SetOnTakeRecorderPanelChanged(FOnTakeRecorderPanelChanged OnTakeRecorderPanelChanged);
 
+	/** Called before initialization occurs (ie. when the recording button is pressed and before the countdown starts) */
+	UFUNCTION(BlueprintCallable, Category = "Take Recorder")
+	static void SetOnTakeRecorderPreInitialize(FOnTakeRecorderPreInitialize OnTakeRecorderPreInitialize);
+
 	/** Called when take recording starts. */
 	UFUNCTION(BlueprintCallable, Category = "Take Recorder")
 	static void SetOnTakeRecorderStarted(FOnTakeRecorderStarted OnTakeRecorderStarted);
@@ -111,6 +116,7 @@ public:
 	UFUNCTION(BlueprintCallable, Category = "Take Recorder")
 	static void SetOnTakeRecorderMarkedFrameAdded(FOnTakeRecorderMarkedFrameAdded OnTakeRecorderMarkedFrameAdded);
 
+	static void OnTakeRecorderPreInitialize();
 	static void OnTakeRecorderStarted();
 	static void OnTakeRecorderStopped();
 	static void OnTakeRecorderFinished(ULevelSequence* InSequenceAsset);
