Take Recorder: Expand playback range when start recording. This used to be when the asset was created but was ugly because you could easily frame the duration of the playback range. Changing this to infinite at the start of recording allows the current time to progress.

diff --git a/Engine/Plugins/VirtualProduction/Takes/Source/TakeRecorder/Private/Recorder/TakeRecorder.cpp b/Engine/Plugins/VirtualProduction/Takes/Source/TakeRecorder/Private/Recorder/TakeRecorder.cpp
--- a/Engine/Plugins/VirtualProduction/Takes/Source/TakeRecorder/Private/Recorder/TakeRecorder.cpp
+++ b/Engine/Plugins/VirtualProduction/Takes/Source/TakeRecorder/Private/Recorder/TakeRecorder.cpp
@@ -719,6 +719,9 @@ void UTakeRecorder::Start(const FTimecode& InTimecodeSource)
 			Sequencer->ResetTimeController();
 		}
 		Sequencer->SetPlaybackStatus(EMovieScenePlayerStatus::Playing);
+
+		// Set infinite playback range when starting recording. Playback range will be clamped to the bounds of the sections at the completion of the recording
+		MovieScene->SetPlaybackRange(MovieScene->GetPlaybackRange().GetLowerBoundValue(), TNumericLimits<int32>::Max() - 1, false);
 	}
 	UTakeRecorderSources* Sources = SequenceAsset->FindMetaData<UTakeRecorderSources>();
 	check(Sources);
