Take Recorder: UI Tweaks

1. Show timecode provider value as the running timecode
2. Show the duration in timecode in the popup recording toast
3. Use the timecode provider frame rate instead of the ini value. So that we're always recording using the timecode provider frame rate.
4. Clicking on the frame rate button opens the timecode settings
5. Added timecode in and timecode out to take metadata

diff --git a/Engine/Config/BaseEditorSettings.ini b/Engine/Config/BaseEditorSettings.ini
--- a/Engine/Config/BaseEditorSettings.ini
+++ b/Engine/Config/BaseEditorSettings.ini
@@ -92,5 +92,4 @@ TeleporterClass=Class'/Script/VREditor.VREditorTeleporter'
 
 [DefaultTakeMetaData TakeMetaData]
 Duration=(SubFrame=0.0)
-FrameRate=(Numerator=60, Denominator=1)
 
diff --git a/Engine/Plugins/VirtualProduction/Takes/Source/TakeRecorder/Private/Recorder/TakeRecorder.cpp b/Engine/Plugins/VirtualProduction/Takes/Source/TakeRecorder/Private/Recorder/TakeRecorder.cpp
--- a/Engine/Plugins/VirtualProduction/Takes/Source/TakeRecorder/Private/Recorder/TakeRecorder.cpp
+++ b/Engine/Plugins/VirtualProduction/Takes/Source/TakeRecorder/Private/Recorder/TakeRecorder.cpp
@@ -163,7 +163,7 @@ private:
 		{
 			ETakeRecorderState NewTakeRecorderState = Recorder->GetState();
 
-			if (NewTakeRecorderState == ETakeRecorderState::CountingDown)
+			if (NewTakeRecorderState == ETakeRecorderState::CountingDown || NewTakeRecorderState == ETakeRecorderState::Started)
 			{
 				// When counting down the text may change on tick
 				TextBlock->SetText(GetDetailText());
@@ -223,6 +223,22 @@ private:
 			{
 				return LOCTEXT("CancelledText", "Recording Cancelled");
 			}
+
+			ULevelSequence* LevelSequence = Recorder->GetSequence();
+
+			UTakeMetaData* TakeMetaData = LevelSequence ? LevelSequence->FindMetaData<UTakeMetaData>() : nullptr;
+
+			if (TakeMetaData)
+			{
+				FFrameRate FrameRate = TakeMetaData->GetFrameRate();
+				FTimespan RecordingDuration = FDateTime::UtcNow() - TakeMetaData->GetTimestamp();
+
+				FFrameNumber TotalFrames = FFrameNumber(static_cast<int32>(FrameRate.AsDecimal() * RecordingDuration.GetTotalSeconds()));
+
+				FTimecode Timecode = FTimecode::FromFrameNumber(TotalFrames, FrameRate, FTimecode::IsDropFormatTimecodeSupported(FrameRate));
+
+				return FText::Format(LOCTEXT("RecordingText", "Recording...{0}"), FText::FromString(Timecode.ToString()));
+			}
 		}
 
 		return LOCTEXT("RecordingText", "Recording...");
@@ -509,6 +525,7 @@ bool UTakeRecorder::CreateDestinationAsset(const TCHAR* AssetPathFormat, ULevelS
 
 	FDateTime UtcNow = FDateTime::UtcNow();
 	AssetMetaData->SetTimestamp(UtcNow);
+	AssetMetaData->SetFrameRate(FApp::GetTimecodeFrameRate());
 
 	// @todo: duration / tick resolution / sample rate / frame rate needs some clarification between sync clocks, template sequences and meta data
 	if (AssetMetaData->GetDuration() > 0)
@@ -729,6 +746,7 @@ void UTakeRecorder::Start(const FTimecode& InTimecodeSource)
 	UTakeMetaData* AssetMetaData = SequenceAsset->FindMetaData<UTakeMetaData>();
 	FDateTime UtcNow = FDateTime::UtcNow();
 	AssetMetaData->SetTimestamp(UtcNow);
+	AssetMetaData->SetTimecodeIn(InTimecodeSource);
 
 	Sources->StartRecording(SequenceAsset, InTimecodeSource, Parameters.User.bAutoSerialize ? &ManifestSerializer : nullptr);
 
@@ -739,6 +757,7 @@ void UTakeRecorder::Start(const FTimecode& InTimecodeSource)
 
 void UTakeRecorder::Stop()
 {
+	FTimecode TimecodeOut = FApp::GetTimecode();
 	USequencerSettings* SequencerSettings = USequencerSettingsContainer::GetOrCreate<USequencerSettings>(TEXT("TakeRecorderSequenceEditor"));
 
 	SequencerSettings->SetAllowEditsMode(CachedAllowEditsMode);
@@ -797,6 +816,8 @@ void UTakeRecorder::Stop()
 		UTakeMetaData* AssetMetaData = SequenceAsset->FindMetaData<UTakeMetaData>();
 		check(AssetMetaData);
 
+		AssetMetaData->SetTimecodeOut(TimecodeOut);
+
 		if (GEditor && GEditor->GetEditorWorldContext().World())
 		{
 			AssetMetaData->SetLevelOrigin(GEditor->GetEditorWorldContext().World()->PersistentLevel);
diff --git a/Engine/Plugins/VirtualProduction/Takes/Source/TakeRecorder/Private/Widgets/STakeRecorderCockpit.cpp b/Engine/Plugins/VirtualProduction/Takes/Source/TakeRecorder/Private/Widgets/STakeRecorderCockpit.cpp
--- a/Engine/Plugins/VirtualProduction/Takes/Source/TakeRecorder/Private/Widgets/STakeRecorderCockpit.cpp
+++ b/Engine/Plugins/VirtualProduction/Takes/Source/TakeRecorder/Private/Widgets/STakeRecorderCockpit.cpp
@@ -51,6 +51,7 @@
 #include "Editor.h"
 #include "Modules/ModuleManager.h"
 #include "LevelEditor.h"
+#include "ISettingsModule.h"
 
 extern UNREALED_API UEditorEngine* GEditor;
 
@@ -329,6 +330,7 @@ void STakeRecorderCockpit::Construct(const FArguments& InArgs)
 						.Font(FTakeRecorderStyle::Get().GetFontStyle("TakeRecorder.Cockpit.SmallText"))
 						.ColorAndOpacity(FSlateColor::UseSubduedForeground())
 						.Text(this, &STakeRecorderCockpit::GetTimestampText)
+						.ToolTipText(this, &STakeRecorderCockpit::GetTimestampTooltipText)
 					]
 
 					+ SHorizontalBox::Slot()
@@ -344,7 +346,8 @@ void STakeRecorderCockpit::Construct(const FArguments& InArgs)
 						.Font(FTakeRecorderStyle::Get().GetFontStyle("TakeRecorder.Cockpit.MediumText"))
 						.ColorAndOpacity(FSlateColor::UseSubduedForeground())
 						.Justification(ETextJustify::Right)
-						.Text(this, &STakeRecorderCockpit::GetDurationText)
+						.Text(this, &STakeRecorderCockpit::GetTimecodeText)
+						.ToolTipText(LOCTEXT("Timecode", "The current timecode"))
 					]
 
 					+SHorizontalBox::Slot()
@@ -377,10 +380,18 @@ void STakeRecorderCockpit::Construct(const FArguments& InArgs)
 					+ SHorizontalBox::Slot()
 					.AutoWidth()
 					[
-						SNew(STextBlock)
-						.ColorAndOpacity(FSlateColor::UseSubduedForeground())
-						.Font(FTakeRecorderStyle::Get().GetFontStyle("TakeRecorder.Cockpit.SmallText"))
-						.Text(this, &STakeRecorderCockpit::GetFrameRateText)
+						SNew(SButton)
+						.ButtonStyle(FEditorStyle::Get(), "NoBorder")
+						.OnClicked(this, &STakeRecorderCockpit::SetFrameRate)
+						.ForegroundColor(FSlateColor::UseForeground())
+						.Content()
+						[
+							SNew(STextBlock)
+							.ColorAndOpacity(FSlateColor::UseSubduedForeground())
+							.Font(FTakeRecorderStyle::Get().GetFontStyle("TakeRecorder.Cockpit.SmallText"))
+							.Text(this, &STakeRecorderCockpit::GetFrameRateText)
+							.ToolTipText(this, &STakeRecorderCockpit::GetFrameRateTooltipText)
+						]
 					]
 				]
 
@@ -596,21 +607,9 @@ FText STakeRecorderCockpit::GetSlateText() const
 	return FText::FromString(TakeMetaData->GetSlate());
 }
 
-FText STakeRecorderCockpit::GetDurationText() const
+FText STakeRecorderCockpit::GetTimecodeText() const
 {
-	FFrameNumber TotalFrames;
-	FFrameRate FrameRate = TakeMetaData->GetFrameRate();
-
-	if (UTakeRecorderBlueprintLibrary::IsRecording())
-	{
-		FTimespan RecordingDuration =  FDateTime::UtcNow() - TakeMetaData->GetTimestamp();
-		TotalFrames = FFrameNumber(static_cast<int32>(FrameRate.AsDecimal() * RecordingDuration.GetTotalSeconds()));
-	}
-
-	FTimecode Timecode = FTimecode::FromFrameNumber(TotalFrames, FrameRate, FTimecode::IsDropFormatTimecodeSupported(FrameRate));
-
-	// FTimecode Timecode = FTimecode::FromFrameNumber(TakeMetaData->GetDuration().FrameNumber, TakeMetaData->GetFrameRate(), FTimecode::IsDropFormatTimecodeSupported(TakeMetaData->GetFrameRate()));
-	return FText::FromString(Timecode.ToString());
+	return FText::FromString(FApp::GetTimecode().ToString());
 }
 
 FText STakeRecorderCockpit::GetUserDescriptionText() const
@@ -620,19 +619,64 @@ FText STakeRecorderCockpit::GetUserDescriptionText() const
 
 FText STakeRecorderCockpit::GetTimestampText() const
 {
-	// FDateTime Timestamp = TakeMetaData->GetTimestamp() == FDateTime(0) ? FDateTime::UtcNow() : TakeMetaData->GetTimestamp();
-	FText TextTime = TakeMetaData->GetTimestamp() == FDateTime(0) ? FText::FromString(TEXT("--")) : FText::AsDateTime(TakeMetaData->GetTimestamp());
-	return TextTime;
+	// If not recorded, return current time
+	if (TakeMetaData->GetTimestamp() == FDateTime(0))
+	{
+		return FText::AsDateTime(FDateTime::UtcNow());
+	}
+	else
+	{
+		return FText::AsDateTime(TakeMetaData->GetTimestamp());
+	}
+}
+
+FText STakeRecorderCockpit::GetTimestampTooltipText() const
+{
+	// If not recorded, return current time
+	if (TakeMetaData->GetTimestamp() == FDateTime(0))
+	{
+		return LOCTEXT("CurrentTimestamp", "The current date/time");
+	}
+	else
+	{
+		return LOCTEXT("Timestamp", "The date/time this recording was created at");
+	}
+}
+
+FReply STakeRecorderCockpit::SetFrameRate()
+{
+	if (ISettingsModule* SettingsModule = FModuleManager::GetModulePtr<ISettingsModule>("Settings"))
+	{
+		SettingsModule->ShowViewer("Project", "Engine", "General");
+		return FReply::Handled();
+	}
+	return FReply::Unhandled();
 }
 
 FText STakeRecorderCockpit::GetFrameRateText() const
 {
-	return GetFrameRate().ToPrettyText();
+	// If not recorded, return app timecode frame rate
+	if (TakeMetaData->GetTimestamp() == FDateTime(0))
+	{
+		return FApp::GetTimecodeFrameRate().ToPrettyText();
+	}
+	else
+	{
+		return TakeMetaData->GetFrameRate().ToPrettyText();
+	}
 }
 
-FFrameRate STakeRecorderCockpit::GetFrameRate() const
+FText STakeRecorderCockpit::GetFrameRateTooltipText() const
 {
-	return TakeMetaData->GetFrameRate();
+	// If not recorded, return app timecode frame rate
+	if (TakeMetaData->GetTimestamp() == FDateTime(0))
+	{
+		return LOCTEXT("ProjectFrameRate", "The project timecode frame rate. The resulting recorded sequence will be at this frame rate.");
+	}
+	else
+	{
+		return LOCTEXT("FrameRate", "The frame rate this recording was created at");
+	}
 }
 
 bool STakeRecorderCockpit::IsFrameRateCompatible(FFrameRate InFrameRate) const
@@ -673,34 +717,6 @@ void STakeRecorderCockpit::SetUserDescriptionText(const FText& InNewText, ETextC
 	}
 }
 
-void STakeRecorderCockpit::SetDurationText(const FText& InNewText, ETextCommit::Type)
-{
-	double CurrentFrameTime = TakeMetaData->GetDuration().AsDecimal();
-
-	FFrameNumberInterface Interface(EFrameNumberDisplayFormats::DropFrameTimecode, 2, TakeMetaData->GetFrameRate(), TakeMetaData->GetFrameRate());
-	TOptional<double> NewFrameTime = Interface.FromString(InNewText.ToString(), CurrentFrameTime);
-
-	if (NewFrameTime.IsSet())
-	{
-		FScopedTransaction Transaction(LOCTEXT("SetDuration_Transaction", "Set Duration"));
-		TakeMetaData->Modify();
-
-		FFrameTime NewDuration = FFrameTime::FromDecimal(NewFrameTime.GetValue());
-		TakeMetaData->SetDuration(NewDuration);
-
-		ULevelSequence* Sequence   = LevelSequenceAttribute.Get();
-		UMovieScene*    MovieScene = Sequence ? Sequence->GetMovieScene() : nullptr;
-
-		if (MovieScene)
-		{
-			MovieScene->Modify();
-
-			TRange<FFrameNumber> PlaybackRange = TRange<FFrameNumber>::Inclusive(0, ConvertFrameTime(NewDuration, TakeMetaData->GetFrameRate(), MovieScene->GetTickResolution()).CeilToFrame());
-			MovieScene->SetPlaybackRange(PlaybackRange);
-		}
-	}
-}
-
 int32 STakeRecorderCockpit::GetTakeNumber() const
 {
 	return TakeMetaData->GetTakeNumber();
diff --git a/Engine/Plugins/VirtualProduction/Takes/Source/TakeRecorder/Private/Widgets/STakeRecorderCockpit.h b/Engine/Plugins/VirtualProduction/Takes/Source/TakeRecorder/Private/Widgets/STakeRecorderCockpit.h
--- a/Engine/Plugins/VirtualProduction/Takes/Source/TakeRecorder/Private/Widgets/STakeRecorderCockpit.h
+++ b/Engine/Plugins/VirtualProduction/Takes/Source/TakeRecorder/Private/Widgets/STakeRecorderCockpit.h
@@ -70,20 +70,21 @@ private:
 	EVisibility GetCountdownVisibility() const;
 	FText GetCountdownText() const;
 
-	FText GetDurationText() const;
-	void SetDurationText(const FText& InNewText, ETextCommit::Type);
+	FText GetTimecodeText() const;
 
 	FText GetUserDescriptionText() const;
 	void SetUserDescriptionText(const FText& InNewText, ETextCommit::Type);
 
+	FReply SetFrameRate();
 	FText GetFrameRateText() const;
-	FFrameRate GetFrameRate() const;
+	FText GetFrameRateTooltipText() const;
 	bool IsFrameRateCompatible(FFrameRate InFrameRate) const;
 
 	FText GetSlateText() const;
 	void SetSlateText(const FText& InNewText, ETextCommit::Type InCommitType);
 
 	FText GetTimestampText() const;
+	FText GetTimestampTooltipText() const;
 
 	int32 GetTakeNumber() const;
 
diff --git a/Engine/Plugins/VirtualProduction/Takes/Source/TakesCore/Private/TakeMetaData.cpp b/Engine/Plugins/VirtualProduction/Takes/Source/TakesCore/Private/TakeMetaData.cpp
--- a/Engine/Plugins/VirtualProduction/Takes/Source/TakesCore/Private/TakeMetaData.cpp
+++ b/Engine/Plugins/VirtualProduction/Takes/Source/TakesCore/Private/TakeMetaData.cpp
@@ -10,7 +10,9 @@
 
 const FName UTakeMetaData::AssetRegistryTag_Slate       = "TakeMetaData_Slate";
 const FName UTakeMetaData::AssetRegistryTag_TakeNumber  = "TakeMetaData_TakeNumber";
-const FName UTakeMetaData::AssetRegistryTag_Timestamp   = "TakeMetaData_Timestamp";
+const FName UTakeMetaData::AssetRegistryTag_Timestamp = "TakeMetaData_Timestamp";
+const FName UTakeMetaData::AssetRegistryTag_TimecodeIn = "TakeMetaData_TimecodeIn";
+const FName UTakeMetaData::AssetRegistryTag_TimecodeOut = "TakeMetaData_TimecodeOut";
 const FName UTakeMetaData::AssetRegistryTag_Description = "TakeMetaData_Description";
 const FName UTakeMetaData::AssetRegistryTag_LevelPath   = "TakeMetaData_LevelPath";
 
@@ -101,6 +103,16 @@ FDateTime UTakeMetaData::GetTimestamp() const
 	return Timestamp;
 }
 
+FTimecode UTakeMetaData::GetTimecodeIn() const
+{
+	return TimecodeIn;
+}
+
+FTimecode UTakeMetaData::GetTimecodeOut() const
+{
+	return TimecodeOut;
+}
+
 FFrameTime UTakeMetaData::GetDuration() const
 {
 	return Duration;
@@ -165,6 +177,22 @@ void UTakeMetaData::SetTimestamp(FDateTime InTimestamp)
 	}
 }
 
+void UTakeMetaData::SetTimecodeIn(FTimecode InTimecodeIn)
+{
+	if (!bIsLocked)
+	{
+		TimecodeIn = InTimecodeIn;
+	}
+}
+
+void UTakeMetaData::SetTimecodeOut(FTimecode InTimecodeOut)
+{
+	if (!bIsLocked)
+	{
+		TimecodeOut = InTimecodeOut;
+	}
+}
+
 void UTakeMetaData::SetDuration(FFrameTime InDuration)
 {
 	if (!bIsLocked)
@@ -209,7 +237,9 @@ void UTakeMetaData::ExtendAssetRegistryTags(TArray<FAssetRegistryTag>& OutTags)
 {
 	OutTags.Emplace(AssetRegistryTag_Slate,       Slate,                     FAssetRegistryTag::ETagType::TT_Alphabetical,  FAssetRegistryTag::TD_None);
 	OutTags.Emplace(AssetRegistryTag_TakeNumber,  LexToString(TakeNumber),   FAssetRegistryTag::ETagType::TT_Numerical,     FAssetRegistryTag::TD_None);
-	OutTags.Emplace(AssetRegistryTag_Timestamp,   Timestamp.ToString(),      FAssetRegistryTag::ETagType::TT_Chronological, FAssetRegistryTag::TD_Date | FAssetRegistryTag::TD_Time);
+	OutTags.Emplace(AssetRegistryTag_Timestamp, Timestamp.ToString(), FAssetRegistryTag::ETagType::TT_Chronological, FAssetRegistryTag::TD_Date | FAssetRegistryTag::TD_Time);
+	OutTags.Emplace(AssetRegistryTag_TimecodeIn, TimecodeIn.ToString(), FAssetRegistryTag::ETagType::TT_Numerical, FAssetRegistryTag::TD_None);
+	OutTags.Emplace(AssetRegistryTag_TimecodeOut, TimecodeOut.ToString(), FAssetRegistryTag::ETagType::TT_Numerical, FAssetRegistryTag::TD_None);
 	OutTags.Emplace(AssetRegistryTag_Description, Description,               FAssetRegistryTag::ETagType::TT_Alphabetical,  FAssetRegistryTag::TD_None);
 }
 
@@ -227,7 +257,17 @@ void UTakeMetaData::ExtendAssetRegistryTagMetaData(TMap<FName, FAssetRegistryTag
 
 	OutMetadata.Add(AssetRegistryTag_Timestamp, FAssetRegistryTagMetadata()
 		.SetDisplayName(NSLOCTEXT("TakeMetaData", "Timestamp_Label", "Timestamp"))
-		.SetTooltip(    NSLOCTEXT("TakeMetaData", "Timestamp_Tip",   "The time that this take was started"))
+		.SetTooltip(NSLOCTEXT("TakeMetaData", "Timestamp_Tip", "The time that this take was started"))
+	);
+
+	OutMetadata.Add(AssetRegistryTag_TimecodeIn, FAssetRegistryTagMetadata()
+		.SetDisplayName(NSLOCTEXT("TakeMetaData", "TimecodeIn_Label", "Timecode In"))
+		.SetTooltip(NSLOCTEXT("TakeMetaData", "TimecodeIn_Tip", "The timecode when this recording was started"))
+	);
+
+	OutMetadata.Add(AssetRegistryTag_TimecodeOut, FAssetRegistryTagMetadata()
+		.SetDisplayName(NSLOCTEXT("TakeMetaData", "TimecodeOut_Label", "Timecode Out"))
+		.SetTooltip(NSLOCTEXT("TakeMetaData", "TimecodeOut_Tip", "The timecode when this recording was stopped"))
 	);
 
 	OutMetadata.Add(AssetRegistryTag_Description, FAssetRegistryTagMetadata()
diff --git a/Engine/Plugins/VirtualProduction/Takes/Source/TakesCore/Public/TakeMetaData.h b/Engine/Plugins/VirtualProduction/Takes/Source/TakesCore/Public/TakeMetaData.h
--- a/Engine/Plugins/VirtualProduction/Takes/Source/TakesCore/Public/TakeMetaData.h
+++ b/Engine/Plugins/VirtualProduction/Takes/Source/TakesCore/Public/TakeMetaData.h
@@ -34,6 +34,12 @@ public:
 	/** The asset registry tag that contains the timestamp for this meta-data */
 	static const FName AssetRegistryTag_Timestamp;
 
+	/** The asset registry tag that contains the timecode in for this meta-data */
+	static const FName AssetRegistryTag_TimecodeIn;
+
+	/** The asset registry tag that contains the timecode out for this meta-data */
+	static const FName AssetRegistryTag_TimecodeOut;
+
 	/** The asset registry tag that contains the user-description for this meta-data */
 	static const FName AssetRegistryTag_Description;
 
@@ -98,6 +104,18 @@ public:
 	UFUNCTION(BlueprintCallable, Category="Take")
 	FDateTime GetTimestamp() const;
 
+	/**
+	 * @return The timecode in for this take
+	 */
+	UFUNCTION(BlueprintCallable, Category = "Take")
+	FTimecode GetTimecodeIn() const;
+
+	/**
+	 * @return The timecode out for this take
+	 */
+	UFUNCTION(BlueprintCallable, Category = "Take")
+	FTimecode GetTimecodeOut() const;
+
 	/**
 	 * @return The duration for this take
 	 */
@@ -177,6 +195,20 @@ public:
 	UFUNCTION(BlueprintCallable, Category="Take")
 	void SetTimestamp(FDateTime InTimestamp);
 
+	/**
+	 * Set this take's timecode in
+	 * @note: Only valid for takes that have not been locked
+	 */
+	UFUNCTION(BlueprintCallable, Category = "Take")
+	void SetTimecodeIn(FTimecode InTimecodeIn);
+
+	/**
+	 * Set this take's timecode out
+	 * @note: Only valid for takes that have not been locked
+	 */
+	UFUNCTION(BlueprintCallable, Category = "Take")
+	void SetTimecodeOut(FTimecode InTimecodeOut);
+
 	/**
 	 * Set this take's duration
 	 * @note: Only valid for takes that have not been locked
@@ -229,6 +261,14 @@ private:
 	UPROPERTY()
 	FDateTime Timestamp;
 
+	/** The timecode at the start of recording */
+	UPROPERTY()
+	FTimecode TimecodeIn;
+
+	/** The timecode at the end of recording */
+	UPROPERTY()
+	FTimecode TimecodeOut;
+
 	/** The desired duration for the take */
 	UPROPERTY(config)
 	FFrameTime Duration;
