partial CIS fix #rb none

diff --git a/Engine/Plugins/VirtualProduction/Takes/Source/TakeRecorder/Private/Recorder/TakeRecorder.cpp b/Engine/Plugins/VirtualProduction/Takes/Source/TakeRecorder/Private/Recorder/TakeRecorder.cpp
--- a/Engine/Plugins/VirtualProduction/Takes/Source/TakeRecorder/Private/Recorder/TakeRecorder.cpp
+++ b/Engine/Plugins/VirtualProduction/Takes/Source/TakeRecorder/Private/Recorder/TakeRecorder.cpp
@@ -40,7 +40,7 @@
 #include "Widgets/Input/SButton.h"
 
 // LevelEditor includes
-#include "ILevelViewport.h"
+#include "IAssetViewport.h"
 #include "LevelEditor.h"
 #include "Subsystems/AssetEditorSubsystem.h"
 
@@ -571,7 +571,7 @@ void UTakeRecorder::InitializeFromParameters()
 	// Apply immersive mode if the parameters demand it
 	if (Parameters.User.bMaximizeViewport)
 	{
-		TSharedPtr<ILevelViewport> ActiveLevelViewport = FModuleManager::Get().LoadModuleChecked<FLevelEditorModule>("LevelEditor").GetFirstActiveViewport();
+		TSharedPtr<IAssetViewport> ActiveLevelViewport = FModuleManager::Get().LoadModuleChecked<FLevelEditorModule>("LevelEditor").GetFirstActiveViewport();
 
 		// If it's already immersive we just leave it alone
 		if (ActiveLevelViewport.IsValid() && !ActiveLevelViewport->IsImmersive())
@@ -579,9 +579,9 @@ void UTakeRecorder::InitializeFromParameters()
 			ActiveLevelViewport->MakeImmersive(true/*bWantImmersive*/, false/*bAllowAnimation*/);
 
 			// Restore it when we're done
-			auto RestoreImmersiveMode = [WeakViewport = TWeakPtr<ILevelViewport>(ActiveLevelViewport)]
+			auto RestoreImmersiveMode = [WeakViewport = TWeakPtr<IAssetViewport>(ActiveLevelViewport)]
 			{
-				if (TSharedPtr<ILevelViewport> CleaupViewport = WeakViewport.Pin())
+				if (TSharedPtr<IAssetViewport> CleaupViewport = WeakViewport.Pin())
 				{
 					CleaupViewport->MakeImmersive(false/*bWantImmersive*/, false/*bAllowAnimation*/);
 				}
