Sequencer: Embed slate and take number into animation asset

diff --git a/Engine/Plugins/VirtualProduction/Takes/Source/TakeRecorderSources/Private/TakeRecorderActorSource.h b/Engine/Plugins/VirtualProduction/Takes/Source/TakeRecorderSources/Private/TakeRecorderActorSource.h
--- a/Engine/Plugins/VirtualProduction/Takes/Source/TakeRecorderSources/Private/TakeRecorderActorSource.h
+++ b/Engine/Plugins/VirtualProduction/Takes/Source/TakeRecorderSources/Private/TakeRecorderActorSource.h
@@ -238,7 +238,8 @@ protected:
 	FTrackRecorderSettings GetTrackRecorderSettings() const override;
 	/** Returns offset that may get set when recording a skeletal mesh animation. Needed to correctly transform attached children*/
 	FTransform GetRecordedActorAnimationInitialRootTransform(class AActor* OtherActor) const override;
-
+	/** Returns the master level sequence being recorded into */
+	ULevelSequence* GetMasterLevelSequence() const override { return MasterLevelSequence; };
 	// ~IMovieSceneTrackRecorderHost Interface
 	
 	/** Initializes an instance of the specified class if we don't already have it in our Settings array. */
diff --git a/Engine/Plugins/VirtualProduction/Takes/Source/TakeTrackRecorders/Private/MovieSceneAnimationTrackRecorder.cpp b/Engine/Plugins/VirtualProduction/Takes/Source/TakeTrackRecorders/Private/MovieSceneAnimationTrackRecorder.cpp
--- a/Engine/Plugins/VirtualProduction/Takes/Source/TakeTrackRecorders/Private/MovieSceneAnimationTrackRecorder.cpp
+++ b/Engine/Plugins/VirtualProduction/Takes/Source/TakeTrackRecorders/Private/MovieSceneAnimationTrackRecorder.cpp
@@ -12,6 +12,7 @@
 #include "Animation/AnimSequence.h"
 #include "Engine/TimecodeProvider.h"
 #include "Engine/Engine.h"
+#include "LevelSequence.h"
 
 DEFINE_LOG_CATEGORY(AnimationSerialization);
 
@@ -40,6 +41,13 @@ void UMovieSceneAnimationTrackRecorder::CreateAnimationAssetAndSequence(const AA
 	{
 		ComponentTransform = SkeletalMeshComponent->GetComponentToWorld().GetRelativeTransform(Actor->GetTransform());
 		FString AnimationAssetName = Actor->GetActorLabel();
+
+		if (ULevelSequence* MasterLevelSequence = OwningTakeRecorderSource->GetMasterLevelSequence())
+		{
+			const FString& MasterLevelSequenceName = MasterLevelSequence->GetName();
+			AnimationAssetName = AnimationAssetName + TEXT("_") + MasterLevelSequenceName;
+		}
+
 		AnimSequence = TakesUtils::MakeNewAsset<UAnimSequence>(AnimationDirectory.Path, AnimationAssetName);
 		if (AnimSequence.IsValid())
 		{
diff --git a/Engine/Plugins/VirtualProduction/Takes/Source/TakeTrackRecorders/Public/TrackRecorders/IMovieSceneTrackRecorderHost.h b/Engine/Plugins/VirtualProduction/Takes/Source/TakeTrackRecorders/Public/TrackRecorders/IMovieSceneTrackRecorderHost.h
--- a/Engine/Plugins/VirtualProduction/Takes/Source/TakeTrackRecorders/Public/TrackRecorders/IMovieSceneTrackRecorderHost.h
+++ b/Engine/Plugins/VirtualProduction/Takes/Source/TakeTrackRecorders/Public/TrackRecorders/IMovieSceneTrackRecorderHost.h
@@ -6,6 +6,7 @@
 #include "MovieSceneSequenceID.h"
 #include "IMovieSceneTrackRecorderHost.generated.h"
 
+class ULevelSequence;
 
 USTRUCT()
 struct FTakeRecorderPropertyTrackSettings
@@ -174,5 +175,8 @@ public:
 	*/
 	virtual FTransform GetRecordedActorAnimationInitialRootTransform(class AActor* OtherActor) const = 0;
 
+	/** Get the master level sequence being recorded into */
+	virtual ULevelSequence* GetMasterLevelSequence() const = 0;
+
 	/**/
 };
