Take Recorder: Remove unused sample clock
#rnx

diff --git a/Engine/Plugins/VirtualProduction/Takes/Source/TakeRecorder/Private/Recorder/TakeRecorderParameters.cpp b/Engine/Plugins/VirtualProduction/Takes/Source/TakeRecorder/Private/Recorder/TakeRecorderParameters.cpp
--- a/Engine/Plugins/VirtualProduction/Takes/Source/TakeRecorder/Private/Recorder/TakeRecorderParameters.cpp
+++ b/Engine/Plugins/VirtualProduction/Takes/Source/TakeRecorder/Private/Recorder/TakeRecorderParameters.cpp
@@ -6,7 +6,6 @@ FTakeRecorderUserParameters::FTakeRecorderUserParameters()
 	: bMaximizeViewport(false)
 	, CountdownSeconds(0.f)
 	, EngineTimeDilation(1.f)
-	, SampleClock(nullptr)
 	, bRemoveRedundantTracks(true)
 	, bSaveRecordedAssets(false)
 	, bAutoSerialize(false)
diff --git a/Engine/Plugins/VirtualProduction/Takes/Source/TakeRecorder/Private/TakeRecorderSettings.cpp b/Engine/Plugins/VirtualProduction/Takes/Source/TakeRecorder/Private/TakeRecorderSettings.cpp
--- a/Engine/Plugins/VirtualProduction/Takes/Source/TakeRecorder/Private/TakeRecorderSettings.cpp
+++ b/Engine/Plugins/VirtualProduction/Takes/Source/TakeRecorder/Private/TakeRecorderSettings.cpp
@@ -8,7 +8,6 @@ UTakeRecorderUserSettings::UTakeRecorderUserSettings()
 	Settings.bMaximizeViewport  = false;
 	Settings.CountdownSeconds   = 3.f;
 	Settings.EngineTimeDilation = 1.f;
-	Settings.SampleClock        = nullptr;
 	Settings.bRemoveRedundantTracks = true;
 	Settings.bSaveRecordedAssets = true;
 	Settings.bAutoSerialize     = false;
diff --git a/Engine/Plugins/VirtualProduction/Takes/Source/TakeRecorder/Public/Recorder/TakeRecorderParameters.h b/Engine/Plugins/VirtualProduction/Takes/Source/TakeRecorder/Public/Recorder/TakeRecorderParameters.h
--- a/Engine/Plugins/VirtualProduction/Takes/Source/TakeRecorder/Public/Recorder/TakeRecorderParameters.h
+++ b/Engine/Plugins/VirtualProduction/Takes/Source/TakeRecorder/Public/Recorder/TakeRecorderParameters.h
@@ -8,12 +8,6 @@
 #include "TakeRecorderParameters.generated.h"
 
 
-UCLASS(Abstract)
-class TAKERECORDER_API UTakeRecorderClock : public UObject
-{
-	GENERATED_BODY()
-};
-
 USTRUCT(BlueprintType)
 struct FTakeRecorderUserParameters
 {
@@ -33,10 +27,6 @@ struct FTakeRecorderUserParameters
 	UPROPERTY(config, EditAnywhere, BlueprintReadWrite, Category="User Settings", meta=(Units=Multiplier, ClampMin="0.00001", UIMin="0.00001"))
 	float EngineTimeDilation;
 
-	/** The clock source to use */
-	UPROPERTY(config, EditAnywhere, BlueprintReadWrite, Category="User Settings", meta=(ShowDisplayNames))
-	TSoftClassPtr<UTakeRecorderClock> SampleClock;
-
 	/** Recommended for use with recorded spawnables. Beware that changes to actor instances in the map after recording may alter the recording when played back */
 	UPROPERTY(config, EditAnywhere, BlueprintReadWrite, Category="User Settings")
 	bool bRemoveRedundantTracks;
