Sequencer: Take number digits now defaults to 2 for better readability.

Take Recorder: Take number digits should respect the sequencer project setting.

diff --git a/Engine/Plugins/VirtualProduction/Takes/Source/TakesCore/Private/TakeMetaData.cpp b/Engine/Plugins/VirtualProduction/Takes/Source/TakesCore/Private/TakeMetaData.cpp
--- a/Engine/Plugins/VirtualProduction/Takes/Source/TakesCore/Private/TakeMetaData.cpp
+++ b/Engine/Plugins/VirtualProduction/Takes/Source/TakesCore/Private/TakeMetaData.cpp
@@ -7,6 +7,7 @@
 #include "ObjectTools.h"
 #include "Engine/Level.h"
 #include "Editor.h"
+#include "MovieSceneToolsProjectSettings.h"
 
 const FName UTakeMetaData::AssetRegistryTag_Slate       = "TakeMetaData_Slate";
 const FName UTakeMetaData::AssetRegistryTag_TakeNumber  = "TakeMetaData_TakeNumber";
@@ -74,6 +75,9 @@ FString UTakeMetaData::GenerateAssetPath(const FString& PathFormatString) const
 	}
 #endif
 
+	const UMovieSceneToolsProjectSettings* ProjectSettings = GetDefault<UMovieSceneToolsProjectSettings>();
+	const int32 TakeNumDigits = ProjectSettings->TakeNumDigits;
+
 	TMap<FString, FStringFormatArg> FormatArgs;
 	FormatArgs.Add(TEXT("day"),    FString::Printf(TEXT("%02i"), TimestampToUse.GetDay()));
 	FormatArgs.Add(TEXT("month"),  FString::Printf(TEXT("%02i"), TimestampToUse.GetMonth()));
@@ -81,7 +85,7 @@ FString UTakeMetaData::GenerateAssetPath(const FString& PathFormatString) const
 	FormatArgs.Add(TEXT("hour"),   FString::Printf(TEXT("%02i"), TimestampToUse.GetHour()));
 	FormatArgs.Add(TEXT("minute"), FString::Printf(TEXT("%02i"), TimestampToUse.GetMinute()));
 	FormatArgs.Add(TEXT("second"), FString::Printf(TEXT("%02i"), TimestampToUse.GetSecond()));
-	FormatArgs.Add(TEXT("take"),   FString::Printf(TEXT("%04i"), TakeNumber));
+	FormatArgs.Add(TEXT("take"),   FString::Printf(TEXT("%0*d"), TakeNumDigits, TakeNumber));
 	FormatArgs.Add(TEXT("slate"),  *Slate);
 	FormatArgs.Add(TEXT("map"),    *MapName);
 
diff --git a/Engine/Source/Editor/MovieSceneTools/Private/MovieSceneToolsProjectSettings.cpp b/Engine/Source/Editor/MovieSceneTools/Private/MovieSceneToolsProjectSettings.cpp
--- a/Engine/Source/Editor/MovieSceneTools/Private/MovieSceneToolsProjectSettings.cpp
+++ b/Engine/Source/Editor/MovieSceneTools/Private/MovieSceneToolsProjectSettings.cpp
@@ -11,7 +11,7 @@ UMovieSceneToolsProjectSettings::UMovieSceneToolsProjectSettings()
 	, FirstShotNumber(10)
 	, ShotIncrement(10)
 	, ShotNumDigits(4)
-	, TakeNumDigits(3)
+	, TakeNumDigits(2)
 	, FirstTakeNumber(1)
 	, TakeSeparator(TEXT("_"))
 	, SubSequenceSeparator(TEXT("_"))
