Update code that was accessing directly the Engine's timecode provider and should have instead use the FApp::GetTimecode.

diff --git a/Engine/Plugins/Animation/LiveLink/Source/LiveLinkSequencer/Private/SequencerRecorderSections/MovieSceneLiveLinkTrackRecorder.cpp b/Engine/Plugins/Animation/LiveLink/Source/LiveLinkSequencer/Private/SequencerRecorderSections/MovieSceneLiveLinkTrackRecorder.cpp
--- a/Engine/Plugins/Animation/LiveLink/Source/LiveLinkSequencer/Private/SequencerRecorderSections/MovieSceneLiveLinkTrackRecorder.cpp
+++ b/Engine/Plugins/Animation/LiveLink/Source/LiveLinkSequencer/Private/SequencerRecorderSections/MovieSceneLiveLinkTrackRecorder.cpp
@@ -194,16 +194,17 @@ void UMovieSceneLiveLinkTrackRecorder::RecordSampleImpl(const FQualifiedFrameTim
 
 		if (FramesToProcess.Num() > 0)
 		{
+			const TOptional<FQualifiedFrameTime> CurrentFrameTime = FApp::GetCurrentFrameTime();
 			for (const FLiveLinkFrameDataStruct& Frame : FramesToProcess)
 			{
 				FFrameNumber FrameNumber;
 
-				if (bSyncedOrForced && GEngine && GEngine->GetTimecodeProvider() && GEngine->GetTimecodeProvider()->GetSynchronizationState() == ETimecodeProviderSynchronizationState::Synchronized)
+				if (bSyncedOrForced && CurrentFrameTime.IsSet())
 				{
 					//Get StartTime on Section in TimeCode FrameRate
 					//Convert that to LiveLink FrameRate and subtract out from LiveLink Frame to get section starting from zero.
 					//Finally convert that to the actual MovieScene Section FrameRate(TickResolution).
-					const FQualifiedFrameTime TimeProviderStartFrameTime = FQualifiedFrameTime(MovieSceneSection->TimecodeSource.Timecode, FApp::GetTimecodeFrameRate());
+					const FQualifiedFrameTime TimeProviderStartFrameTime = FQualifiedFrameTime(MovieSceneSection->TimecodeSource.Timecode, CurrentFrameTime.GetValue().Rate);
 					FQualifiedFrameTime LiveLinkFrameTime = Frame.GetBaseData()->MetaData.SceneTime;
 					const FFrameNumber FrameNumberStart = TimeProviderStartFrameTime.ConvertTo(LiveLinkFrameTime.Rate).FrameNumber;
 					LiveLinkFrameTime.Time.FrameNumber -= FrameNumberStart;
diff --git a/Engine/Plugins/VirtualProduction/Takes/Source/TakeTrackRecorders/Private/MovieSceneAnimationTrackRecorder.cpp b/Engine/Plugins/VirtualProduction/Takes/Source/TakeTrackRecorders/Private/MovieSceneAnimationTrackRecorder.cpp
--- a/Engine/Plugins/VirtualProduction/Takes/Source/TakeTrackRecorders/Private/MovieSceneAnimationTrackRecorder.cpp
+++ b/Engine/Plugins/VirtualProduction/Takes/Source/TakeTrackRecorders/Private/MovieSceneAnimationTrackRecorder.cpp
@@ -121,9 +121,8 @@ void UMovieSceneAnimationTrackRecorder::CreateTrackImpl()
 		{
 			//If we are syncing to a timecode provider use that's frame rate as our frame rate since
 			//otherwise use the displayrate.
-			const UTimecodeProvider* TimecodeProvider = GEngine->GetTimecodeProvider();
-			FFrameRate SampleRate = (TimecodeProvider && TimecodeProvider->GetSynchronizationState() == ETimecodeProviderSynchronizationState::Synchronized)
-				? TimecodeProvider->GetFrameRate() : MovieScene->GetDisplayRate();
+			const TOptional<FQualifiedFrameTime> CurrentFrameTime = FApp::GetCurrentFrameTime();
+			FFrameRate SampleRate = CurrentFrameTime.IsSet() ? CurrentFrameTime.GetValue().Rate : MovieScene->GetDisplayRate();
 
 			FText Error;
 			FString Name = SkeletalMeshComponent->GetName();
@@ -261,9 +260,8 @@ void UMovieSceneAnimationTrackRecorder::RecordSampleImpl(const FQualifiedFrameTi
 			bRecordInWorldSpace = !OwningTakeRecorderSource->IsOtherActorBeingRecorded(AttachParent->GetOwner());
 		}
 
-		const UTimecodeProvider* TimecodeProvider = GEngine->GetTimecodeProvider();
-		FFrameRate SampleRate = (TimecodeProvider && TimecodeProvider->GetSynchronizationState() == ETimecodeProviderSynchronizationState::Synchronized)
-			? TimecodeProvider->GetFrameRate() : MovieScene->GetDisplayRate();
+		const TOptional<FQualifiedFrameTime> CurrentFrameTime = FApp::GetCurrentFrameTime();
+		FFrameRate SampleRate = CurrentFrameTime.IsSet() ? CurrentFrameTime.GetValue().Rate : MovieScene->GetDisplayRate();
 
 		//Set this up here so we know that it's parent sources have also been added so we record in the correct space
 		FAnimationRecordingSettings RecordingSettings;
diff --git a/Engine/Plugins/VirtualProduction/Takes/Source/TakesCore/Private/TakeRecorderSources.cpp b/Engine/Plugins/VirtualProduction/Takes/Source/TakesCore/Private/TakeRecorderSources.cpp
--- a/Engine/Plugins/VirtualProduction/Takes/Source/TakesCore/Private/TakeRecorderSources.cpp
+++ b/Engine/Plugins/VirtualProduction/Takes/Source/TakesCore/Private/TakeRecorderSources.cpp
@@ -568,31 +568,27 @@ FQualifiedFrameTime UTakeRecorderSources::GetCurrentRecordingFrameTime(const FTi
 
 	if (GEngine)
 	{
-		const UTimecodeProvider* TimecodeProvider = GEngine->GetTimecodeProvider();
-
 		// If their is a a Timecode Provider that is synchronized then we will sample the engine Timecode
 		// to determine what frame the data should go on. If the engine is ticking faster than the given
 		// Timecode framerate then there will be multiple frames submitted with the same qualified time
 		// and the data sources will end up only storing the latest call on that frame.
 
-		if (TimecodeProvider && TimecodeProvider->GetSynchronizationState() == ETimecodeProviderSynchronizationState::Synchronized)
+		const TOptional<FQualifiedFrameTime> CurrentFrameTime = FApp::GetCurrentFrameTime();
+		if (CurrentFrameTime.IsSet())
 		{
-			FTimecode Timecode = FApp::GetTimecode();
-			const FFrameNumber QualifiedFrameNumber = Timecode.ToFrameNumber(FApp::GetTimecodeFrameRate());
-			FrameTime = FQualifiedFrameTime(FFrameTime(QualifiedFrameNumber), FApp::GetTimecodeFrameRate());
-
+			FrameTime = CurrentFrameTime.GetValue();
 			bHasValidTimecodeSource = true;
 		}
 		else
 		{
-			UE_LOG(LogTakesCore, Error, TEXT("Attempted to sample timecode from custom Timecode Provider %s while provider was not synchronized! Falling back to engine clock for timecode source!"), TimecodeProvider != nullptr ? *TimecodeProvider->GetName() : TEXT(""));
+			UE_LOG(LogTakesCore, Error, TEXT("Attempted to sample timecode from the engine and failed! Falling back to engine clock for timecode source!"));
 			bHasValidTimecodeSource = false;
 		}
 	}
 
 	if (!bHasValidTimecodeSource)
 	{
-		// If no Timecode Provider is specified (or it has an error) then we want to fall back to the normal Engine tickrate and capture
+		// If no Timecode Provider is specified (or it has an error) then we want to fall back to the normal Engine tick rate and capture
 		//Use Level Sequence TickRate to make conversions cleaner later on.
 		const FFrameNumber FrameNumber = TargetLevelSequenceTickResolution.AsFrameNumber(TimeSinceRecordingStarted);
 		FrameTime = FQualifiedFrameTime(FFrameTime(FrameNumber), TargetLevelSequenceTickResolution);
diff --git a/Engine/Source/Runtime/MovieScene/Private/MovieSceneTimeController.cpp b/Engine/Source/Runtime/MovieScene/Private/MovieSceneTimeController.cpp
--- a/Engine/Source/Runtime/MovieScene/Private/MovieSceneTimeController.cpp
+++ b/Engine/Source/Runtime/MovieScene/Private/MovieSceneTimeController.cpp
@@ -110,12 +110,10 @@ double FMovieSceneTimeController_AudioClock::GetCurrentTime() const
 
 double FMovieSceneTimeController_TimecodeClock::GetCurrentTime() const
 {
-	if (GEngine && GEngine->GetTimecodeProvider() && GEngine->GetTimecodeProvider()->GetSynchronizationState() == ETimecodeProviderSynchronizationState::Synchronized)
+	const TOptional<FQualifiedFrameTime> CurrentFrameTime = FApp::GetCurrentFrameTime();
+	if (CurrentFrameTime.IsSet())
 	{
-		FTimecode Timecode = FApp::GetTimecode();
-		FFrameRate FrameRate = FApp::GetTimecodeFrameRate();
-		FFrameNumber FrameNumber = Timecode.ToFrameNumber(FrameRate);
-		return FrameRate.AsSeconds(FrameNumber);
+		return CurrentFrameTime.GetValue().AsSeconds();
 	}
 	else
 	{
