Take Recoder: Recorded data isn't smooth.
Always use tick/delta time to avoid having duplicate times.

diff --git a/Engine/Plugins/VirtualProduction/Takes/Source/TakesCore/Private/TakeRecorderSources.cpp b/Engine/Plugins/VirtualProduction/Takes/Source/TakesCore/Private/TakeRecorderSources.cpp
--- a/Engine/Plugins/VirtualProduction/Takes/Source/TakesCore/Private/TakeRecorderSources.cpp
+++ b/Engine/Plugins/VirtualProduction/Takes/Source/TakesCore/Private/TakeRecorderSources.cpp
@@ -224,24 +224,13 @@ void UTakeRecorderSources::StartRecordingRecursive(TArray<UTakeRecorderSource*>
 		StartRecordingRecursive(NewSources, InMasterSequence, Timecode, InManifestSerializer);
 		SourcesSerialNumber++;
 
-		bool bHasValidTimecodeSource;
-		FQualifiedFrameTime QualifiedSequenceTime = GetCurrentRecordingFrameTime(Timecode, bHasValidTimecodeSource);
+		FQualifiedFrameTime QualifiedSequenceTime = GetCurrentRecordingFrameTime();
 		for (auto NewSource : NewSources)
 		{
 			if (NewSource->bEnabled)
 			{
 				ULevelSequence* SourceSequence = SourceSubSequenceMap[NewSource];
 				FFrameNumber FrameNumber = QualifiedSequenceTime.ConvertTo(SourceSequence->GetMovieScene()->GetTickResolution()).FloorToFrame();
-				if (bHasValidTimecodeSource)
-				{
-					//Need to get difference of the source time from the start of recording and put into level frame rate.
-					FFrameNumber SeqStartFrameTime = StartRecordingTimecodeSource.ToFrameNumber(FApp::GetTimecodeFrameRate());
-					FFrameNumber StartFrameTime = Timecode.ToFrameNumber(FApp::GetTimecodeFrameRate());
-					FFrameTime FrameTimeDiff;
-					FrameTimeDiff.FrameNumber = StartFrameTime - SeqStartFrameTime;
-					FrameTimeDiff = FFrameRate::TransformTime(FrameTimeDiff, FApp::GetTimecodeFrameRate(), TargetLevelSequenceTickResolution);
-					FrameNumber = FrameTimeDiff.FrameNumber;
-				}
 				NewSource->StartRecording(Timecode, FrameNumber, SourceSubSequenceMap[NewSource]);
 			}
 		}
@@ -426,8 +415,7 @@ void UTakeRecorderSources::StartRecordingSource(TArray<UTakeRecorderSource *> In
 
 void UTakeRecorderSources::StartRecordingTheseSources(const TArray<UTakeRecorderSource *>& InSources, const FTimecode& CurrentTimecode)
 {
-	bool bHasValidTimecodeSource;
-	FQualifiedFrameTime QualifiedSequenceTime = GetCurrentRecordingFrameTime(CurrentTimecode, bHasValidTimecodeSource);
+	FQualifiedFrameTime QualifiedSequenceTime = GetCurrentRecordingFrameTime();
 	for (auto Source : InSources)
 	{
 		if (Source->bEnabled)
@@ -439,22 +427,11 @@ void UTakeRecorderSources::StartRecordingTheseSources(const TArray<UTakeRecorder
 			}
 			FFrameNumber FrameNumber = QualifiedSequenceTime.ConvertTo(SourceSequence->GetMovieScene()->GetTickResolution()).FloorToFrame();
 			Source->TimecodeSource = CurrentTimecode;
-			if (bHasValidTimecodeSource)
-			{
-				//Need to get difference of the source time from the start of recording and put into level frame rate.
-				FFrameNumber SeqStartFrameTime = StartRecordingTimecodeSource.ToFrameNumber(FApp::GetTimecodeFrameRate());
-				FFrameNumber StartFrameTime = CurrentTimecode.ToFrameNumber(FApp::GetTimecodeFrameRate());
-				FFrameTime FrameTimeDiff;
-				FrameTimeDiff.FrameNumber = StartFrameTime - SeqStartFrameTime;
-				FrameTimeDiff = FFrameRate::TransformTime(FrameTimeDiff, FApp::GetTimecodeFrameRate(), TargetLevelSequenceTickResolution);
-				FrameNumber = FrameTimeDiff.FrameNumber;
-			}
 			Source->StartRecording(CurrentTimecode, FrameNumber, SourceSubSequenceMap[Source]);
 		}
 	}
 }
 
-
 void UTakeRecorderSources::PreRecording(class ULevelSequence* InSequence, FManifestSerializer* InManifestSerializer)
 {
 	// We want to cache the Serializer and Level Sequence in case more objects start recording mid-recording.
@@ -472,7 +449,6 @@ void UTakeRecorderSources::StartRecording(class ULevelSequence* InSequence, cons
 
 	bIsRecording = true;
 	TimeSinceRecordingStarted = 0.f;
-	LastTimecodeFrameNumber.Reset();
 	TargetLevelSequenceTickResolution = InSequence->GetMovieScene()->GetTickResolution();
 
 	InSequence->GetMovieScene()->TimecodeSource = InTimecodeSource;
@@ -482,25 +458,10 @@ void UTakeRecorderSources::StartRecording(class ULevelSequence* InSequence, cons
 
 FFrameTime UTakeRecorderSources::TickRecording(class ULevelSequence* InSequence, const FTimecode& InTimecodeSource, float DeltaTime)
 {
-	bool bHasValidTimecodeSource;
-	FQualifiedFrameTime FrameTime = GetCurrentRecordingFrameTime(InTimecodeSource, bHasValidTimecodeSource);
+	FQualifiedFrameTime FrameTime = GetCurrentRecordingFrameTime();
 	FQualifiedFrameTime SourceFrameTime(FrameTime);
-	bool bTimeIncremented = (DeltaTime > 0.0f || !LastTimecodeFrameNumber.IsSet());
-	if (bHasValidTimecodeSource)
-	{
-		//We leave this ins timecode frame rate since the sources convert it later (cbb and faster to do it here, we actually do it below
-		//for showiung the time
-		FFrameNumber SeqStartFrameTime = StartRecordingTimecodeSource.ToFrameNumber(FApp::GetTimecodeFrameRate());
-		SourceFrameTime.Time.FrameNumber = SourceFrameTime.Time.FrameNumber - SeqStartFrameTime;
-		if (LastTimecodeFrameNumber.IsSet())
-		{
-			if (LastTimecodeFrameNumber == SourceFrameTime.Time.FrameNumber)
-			{
-				bTimeIncremented = false;
-			}
-		}
-		LastTimecodeFrameNumber = SourceFrameTime.Time.FrameNumber;
-	}
+	bool bTimeIncremented = DeltaTime > 0.0f;
+
 	if (bTimeIncremented) //only record if time incremented, may not with timecode providers with low frame rates
 	{
 		for (auto Source : Sources)
@@ -515,21 +476,9 @@ FFrameTime UTakeRecorderSources::TickRecording(class ULevelSequence* InSequence,
 	//Time in seconds since recording started. Used when there is no Timecode Sync (e.g. in case it get's lost or dropped).
 	TimeSinceRecordingStarted += DeltaTime;
 
-	//We calculate and return the Current Frame Number based upon whether driven by timecode or engine tick.
-	//We need to make sure this is TargetLevelSequenceTickResolution, but first do in Timecode rate space do
-	//to precision issues with Timecode.
 	FFrameTime CurrentFrameTimeSinceStart;
-	if (bHasValidTimecodeSource)
-	{
-		FFrameNumber SeqStartFrameTime = StartRecordingTimecodeSource.ToFrameNumber(FApp::GetTimecodeFrameRate());
-		FrameTime.Time.FrameNumber = FrameTime.Time.FrameNumber - SeqStartFrameTime;
-		CurrentFrameTimeSinceStart = FrameTime.ConvertTo(TargetLevelSequenceTickResolution);
+	CurrentFrameTimeSinceStart = TargetLevelSequenceTickResolution.AsFrameTime(TimeSinceRecordingStarted);
 
-	}
-	else
-	{
-		CurrentFrameTimeSinceStart = TargetLevelSequenceTickResolution.AsFrameTime(TimeSinceRecordingStarted);
-	}
 
 	// If we're recording into sub-sections we want to update their range every frame so they appear to
 	// animate as their contents are filled. We can't check against the size of all sections (not all
@@ -558,41 +507,13 @@ FFrameTime UTakeRecorderSources::TickRecording(class ULevelSequence* InSequence,
 	return CurrentFrameTimeSinceStart;
 }
 
-//If we have a valid timecode source the returned Qualified Time is the raw Timecode converted time.  Since we will need to
-//convert that time to some relative value, based on start of the level sequence or start of the source we also pass if timecode was valid.
-//If not valid then it's just TimeSinceRecordingStarted converted to a Qualified Time
-FQualifiedFrameTime UTakeRecorderSources::GetCurrentRecordingFrameTime(const FTimecode& InTimecode, bool& bHasValidTimecodeSource) const
+//We now always just use TiomeSinceRecordingStarted insteaad of possibly using timecode to determine our time since
+//That can give us a higher resolution
+FQualifiedFrameTime UTakeRecorderSources::GetCurrentRecordingFrameTime() const
 {
 	FQualifiedFrameTime FrameTime;
-	bHasValidTimecodeSource = false;
-
-	if (GEngine)
-	{
-		// If their is a a Timecode Provider that is synchronized then we will sample the engine Timecode
-		// to determine what frame the data should go on. If the engine is ticking faster than the given
-		// Timecode framerate then there will be multiple frames submitted with the same qualified time
-		// and the data sources will end up only storing the latest call on that frame.
-
-		const TOptional<FQualifiedFrameTime> CurrentFrameTime = FApp::GetCurrentFrameTime();
-		if (CurrentFrameTime.IsSet())
-		{
-			FrameTime = CurrentFrameTime.GetValue();
-			bHasValidTimecodeSource = true;
-		}
-		else
-		{
-			UE_LOG(LogTakesCore, Error, TEXT("Attempted to sample timecode from the engine and failed! Falling back to engine clock for timecode source!"));
-			bHasValidTimecodeSource = false;
-		}
-	}
-
-	if (!bHasValidTimecodeSource)
-	{
-		// If no Timecode Provider is specified (or it has an error) then we want to fall back to the normal Engine tick rate and capture
-		//Use Level Sequence TickRate to make conversions cleaner later on.
-		const FFrameNumber FrameNumber = TargetLevelSequenceTickResolution.AsFrameNumber(TimeSinceRecordingStarted);
-		FrameTime = FQualifiedFrameTime(FFrameTime(FrameNumber), TargetLevelSequenceTickResolution);
-	}
+	const FFrameNumber FrameNumber = TargetLevelSequenceTickResolution.AsFrameNumber(TimeSinceRecordingStarted);
+	FrameTime = FQualifiedFrameTime(FFrameTime(FrameNumber), TargetLevelSequenceTickResolution);
 
 	return FrameTime;
 }
@@ -601,7 +522,6 @@ void UTakeRecorderSources::StopRecording(class ULevelSequence* InSequence, FTake
 {
 	bIsRecording = false;
 	TimeSinceRecordingStarted = 0.f;
-	LastTimecodeFrameNumber.Reset();
 
 	for (auto Source : Sources)
 	{
diff --git a/Engine/Plugins/VirtualProduction/Takes/Source/TakesCore/Public/TakeRecorderSources.h b/Engine/Plugins/VirtualProduction/Takes/Source/TakesCore/Public/TakeRecorderSources.h
--- a/Engine/Plugins/VirtualProduction/Takes/Source/TakesCore/Public/TakeRecorderSources.h
+++ b/Engine/Plugins/VirtualProduction/Takes/Source/TakesCore/Public/TakeRecorderSources.h
@@ -168,8 +168,8 @@ private:
 	/** Finds the folder that the given Source should be created in, creating it if necessary. */
 	class UMovieSceneFolder* AddFolderForSource(const UTakeRecorderSource* InSource, class UMovieScene* InMovieScene);
 
-	/** Gets the current frame time for recording, optionally resolving out the engine's custom Timecode provider. */
-	FQualifiedFrameTime GetCurrentRecordingFrameTime(const FTimecode& InTimeCode, bool& bHasValidTimeCodeSource) const;
+	/** Gets the current frame time for recording */
+	FQualifiedFrameTime GetCurrentRecordingFrameTime() const;
 
 	/** Remove object bindings that don't have any tracks and are not bindings for attach/path tracks */
 	void RemoveRedundantTracks();
@@ -221,9 +221,6 @@ private:
 	/** Timecode time at start of recording */
 	FTimecode StartRecordingTimecodeSource;
 
-	/** Last Timecode Frame Number, used to avoid recording same time twice*/
-	TOptional<FFrameNumber> LastTimecodeFrameNumber;
-
 	/** All sources after PreRecord */
 	TArray<UTakeRecorderSource *> PreRecordedSources;
 
