Take Recorder: Rename NewTake to ClearPendingTake. Users just want the ability to clear the pending level sequence take rather than also clearing out the sources.

diff --git a/Engine/Plugins/VirtualProduction/Takes/Source/TakeRecorder/Private/Recorder/TakeRecorderPanel.cpp b/Engine/Plugins/VirtualProduction/Takes/Source/TakeRecorder/Private/Recorder/TakeRecorderPanel.cpp
--- a/Engine/Plugins/VirtualProduction/Takes/Source/TakeRecorder/Private/Recorder/TakeRecorderPanel.cpp
+++ b/Engine/Plugins/VirtualProduction/Takes/Source/TakeRecorder/Private/Recorder/TakeRecorderPanel.cpp
@@ -73,11 +73,11 @@ void UTakeRecorderPanel::SetupForViewing(ULevelSequence* LevelSequenceAsset)
 	}
 }
 
-void UTakeRecorderPanel::NewTake()
+void UTakeRecorderPanel::ClearPendingTake()
 {
 	if (ValidateTabContent())
 	{
-		WeakTabContent.Pin()->NewTake();
+		WeakTabContent.Pin()->ClearPendingTake();
 	}
 }
 
diff --git a/Engine/Plugins/VirtualProduction/Takes/Source/TakeRecorder/Private/Widgets/STakeRecorderPanel.cpp b/Engine/Plugins/VirtualProduction/Takes/Source/TakeRecorder/Private/Widgets/STakeRecorderPanel.cpp
--- a/Engine/Plugins/VirtualProduction/Takes/Source/TakeRecorder/Private/Widgets/STakeRecorderPanel.cpp
+++ b/Engine/Plugins/VirtualProduction/Takes/Source/TakeRecorder/Private/Widgets/STakeRecorderPanel.cpp
@@ -280,10 +280,10 @@ TSharedRef<SWidget> STakeRecorderPanel::MakeToolBar()
 				SNew(SButton)
 				.HAlign(HAlign_Center)
 				.VAlign(VAlign_Center)
-				.ToolTipText(LOCTEXT("Add", "Create a New Take"))
+				.ToolTipText(LOCTEXT("ClearPendingTake", "Clear Pending Take"))
 				.ForegroundColor(FSlateColor::UseForeground())
 				.ButtonStyle(FEditorStyle::Get(), "HoverHintOnly")
-				.OnClicked(this, &STakeRecorderPanel::OnNewTake)
+				.OnClicked(this, &STakeRecorderPanel::OnClearPendingTake)
 				[
 					SNew(STextBlock)
 					.Font(FEditorStyle::Get().GetFontStyle("FontAwesome.14"))
@@ -448,19 +448,32 @@ UTakeMetaData* STakeRecorderPanel::GetTakeMetaData() const
 	return CockpitWidget->GetMetaData();
 }
 
-void STakeRecorderPanel::NewTake()
+void STakeRecorderPanel::ClearPendingTake()
 {
 	if (CockpitWidget->Reviewing())
 	{
 		LastRecordedLevelSequence = SuppliedLevelSequence;
 	}
 
+	UTakeRecorderSources* BaseSources = nullptr;
+
+	if (ULevelSequence* CurrentLevelSequence = GetLevelSequence())
+	{
+		BaseSources = CurrentLevelSequence->FindMetaData<UTakeRecorderSources>();
+	}
+
 	SuppliedLevelSequence = nullptr;
 
-	FScopedTransaction Transaction(LOCTEXT("NewTake_Transaction", "New Take"));
+	FScopedTransaction Transaction(LOCTEXT("ClearPendingTake_Transaction", "Clear Pending Take"));
 
 	TransientPreset->Modify();
 	TransientPreset->CreateLevelSequence();
+
+	ULevelSequence* LevelSequence = TransientPreset->GetLevelSequence();
+	if (LevelSequence && BaseSources)
+	{
+		LevelSequence->CopyMetaData(BaseSources);
+	}
 }
 
 UTakePreset* STakeRecorderPanel::AllocateTransientPreset()
@@ -708,15 +721,15 @@ FReply STakeRecorderPanel::OnBackToPendingTake()
 	return FReply::Handled();
 }
 
-FReply STakeRecorderPanel::OnNewTake()
+FReply STakeRecorderPanel::OnClearPendingTake()
 {
-	FText WarningMessage (LOCTEXT("Warning_NewTake", "Are you sure you want to create a new empty take setup? Your current changes will be discarded."));
+	FText WarningMessage (LOCTEXT("Warning_ClearPendingTake", "Are you sure you want to clear the pending take? Your current tracks will be discarded."));
 	if (EAppReturnType::No == FMessageDialog::Open(EAppMsgType::YesNo, WarningMessage))
 	{
 		return FReply::Handled();
 	}
 
-	NewTake();
+	ClearPendingTake();
 	return FReply::Handled();
 }
 
diff --git a/Engine/Plugins/VirtualProduction/Takes/Source/TakeRecorder/Private/Widgets/STakeRecorderPanel.h b/Engine/Plugins/VirtualProduction/Takes/Source/TakeRecorder/Private/Widgets/STakeRecorderPanel.h
--- a/Engine/Plugins/VirtualProduction/Takes/Source/TakeRecorder/Private/Widgets/STakeRecorderPanel.h
+++ b/Engine/Plugins/VirtualProduction/Takes/Source/TakeRecorder/Private/Widgets/STakeRecorderPanel.h
@@ -59,7 +59,7 @@ public:
 
 	TSharedPtr<STakeRecorderCockpit> GetCockpitWidget() const { return CockpitWidget; }
 
-	void NewTake();
+	void ClearPendingTake();
 
 private:
 
@@ -96,7 +96,7 @@ private:
 
 	FReply OnBackToPendingTake();
 
-	FReply OnNewTake();
+	FReply OnClearPendingTake();
 
 	FReply OnReviewLastRecording();
 
diff --git a/Engine/Plugins/VirtualProduction/Takes/Source/TakeRecorder/Private/Widgets/STakeRecorderTabContent.cpp b/Engine/Plugins/VirtualProduction/Takes/Source/TakeRecorder/Private/Widgets/STakeRecorderTabContent.cpp
--- a/Engine/Plugins/VirtualProduction/Takes/Source/TakeRecorder/Private/Widgets/STakeRecorderTabContent.cpp
+++ b/Engine/Plugins/VirtualProduction/Takes/Source/TakeRecorder/Private/Widgets/STakeRecorderTabContent.cpp
@@ -303,11 +303,11 @@ void STakeRecorderTabContent::StopRecording() const
 }
 
 
-void STakeRecorderTabContent::NewTake()
+void STakeRecorderTabContent::ClearPendingTake()
 {
 	if (TSharedPtr<STakeRecorderPanel> Panel = WeakPanel.Pin())
 	{
-		Panel->NewTake();
+		Panel->ClearPendingTake();
 	}
 }
 
diff --git a/Engine/Plugins/VirtualProduction/Takes/Source/TakeRecorder/Private/Widgets/STakeRecorderTabContent.h b/Engine/Plugins/VirtualProduction/Takes/Source/TakeRecorder/Private/Widgets/STakeRecorderTabContent.h
--- a/Engine/Plugins/VirtualProduction/Takes/Source/TakeRecorder/Private/Widgets/STakeRecorderTabContent.h
+++ b/Engine/Plugins/VirtualProduction/Takes/Source/TakeRecorder/Private/Widgets/STakeRecorderTabContent.h
@@ -55,7 +55,7 @@ public:
 
 	void StopRecording() const;
 
-	void NewTake();
+	void ClearPendingTake();
 
 private:
 
diff --git a/Engine/Plugins/VirtualProduction/Takes/Source/TakeRecorder/Public/Recorder/TakeRecorderPanel.h b/Engine/Plugins/VirtualProduction/Takes/Source/TakeRecorder/Public/Recorder/TakeRecorderPanel.h
--- a/Engine/Plugins/VirtualProduction/Takes/Source/TakeRecorder/Public/Recorder/TakeRecorderPanel.h
+++ b/Engine/Plugins/VirtualProduction/Takes/Source/TakeRecorder/Public/Recorder/TakeRecorderPanel.h
@@ -65,11 +65,15 @@ public:
 	UFUNCTION(BlueprintCallable, Category="Take Recorder|Panel", DisplayName="Set Mode (Read-Only Level Sequence)")
 	void SetupForViewing(ULevelSequence* LevelSequenceAsset);
 
+	UE_DEPRECATED(4.25, "Please use ClearPendingTake instead")
+	UFUNCTION(BlueprintCallable, Category = "Take Recorder|Panel", DisplayName = "New Take", meta = (DeprecatedFunction, DeprecationMessage = "Please use ClearPendingTake instead"))
+	void NewTake() { ClearPendingTake(); }
+
 	/*
-	 * Setup for a new take by clearing out sources and the transient level sequence
+	 * Clear the pending take level sequence
 	 */
-	UFUNCTION(BlueprintCallable, Category = "Take Recorder|Panel", DisplayName = "New Take")
-	void NewTake();
+	UFUNCTION(BlueprintCallable, Category = "Take Recorder|Panel", DisplayName = "Clear Pending Take")
+	void ClearPendingTake();
 
 	/**
 	 * Access the level sequence for this take
