Take Recorder: Added digits type interface so that take number can be padded

diff --git a/Engine/Plugins/VirtualProduction/Takes/Source/TakeRecorder/Private/Widgets/STakeRecorderCockpit.cpp b/Engine/Plugins/VirtualProduction/Takes/Source/TakeRecorder/Private/Widgets/STakeRecorderCockpit.cpp
--- a/Engine/Plugins/VirtualProduction/Takes/Source/TakeRecorder/Private/Widgets/STakeRecorderCockpit.cpp
+++ b/Engine/Plugins/VirtualProduction/Takes/Source/TakeRecorder/Private/Widgets/STakeRecorderCockpit.cpp
@@ -19,6 +19,8 @@
 #include "Algo/Find.h"
 #include "Framework/MultiBox/MultiBoxBuilder.h"
 
+#include "MovieSceneToolsProjectSettings.h"
+
 // AssetRegistry includes
 #include "AssetRegistryModule.h"
 
@@ -92,6 +94,29 @@ struct SNonThrottledButton : SButton
 	}
 };
 
+struct FDigitsTypeInterface : INumericTypeInterface<int32>
+{
+	virtual FString ToString(const int32& Value) const override
+	{
+		const UMovieSceneToolsProjectSettings* ProjectSettings = GetDefault<UMovieSceneToolsProjectSettings>();
+
+		return FString::Printf(TEXT("%0*d"), ProjectSettings->TakeNumDigits, Value);
+	}
+
+	virtual TOptional<int32> FromString(const FString& InString, const int32& ExistingValue) override
+	{
+		return (int32)FCString::Atoi(*InString);
+	}
+
+	virtual int32 GetMinFractionalDigits() const override { return 0; }
+	virtual int32 GetMaxFractionalDigits() const override { return 0; }
+
+	virtual void SetMinFractionalDigits(const TAttribute<TOptional<int32>>& NewValue) override {}
+	virtual void SetMaxFractionalDigits(const TAttribute<TOptional<int32>>& NewValue) override {}
+
+	virtual bool IsCharacterValid(TCHAR InChar) const override { return true; }
+};
+	
 void STakeRecorderCockpit::Construct(const FArguments& InArgs)
 {
 	FAssetRegistryModule& AssetRegistryModule = FModuleManager::LoadModuleChecked<FAssetRegistryModule>(TEXT("AssetRegistry"));
@@ -121,6 +146,8 @@ void STakeRecorderCockpit::Construct(const FArguments& InArgs)
 	UpdateRecordError();
 
 	CommandList = MakeShareable(new FUICommandList);
+	
+	DigitsTypeInterface = MakeShareable(new FDigitsTypeInterface);
 
 	BindCommands();
 
@@ -235,6 +262,7 @@ void STakeRecorderCockpit::Construct(const FArguments& InArgs)
 						.OnValueChanged(this, &STakeRecorderCockpit::SetTakeNumber)
 						.OnValueCommitted(this, &STakeRecorderCockpit::SetTakeNumber_FromCommit)
 						.OnEndSliderMovement(this, &STakeRecorderCockpit::OnEndSetTakeNumber)
+						.TypeInterface(DigitsTypeInterface)
 					]
 				]
 
diff --git a/Engine/Plugins/VirtualProduction/Takes/Source/TakeRecorder/Private/Widgets/STakeRecorderCockpit.h b/Engine/Plugins/VirtualProduction/Takes/Source/TakeRecorder/Private/Widgets/STakeRecorderCockpit.h
--- a/Engine/Plugins/VirtualProduction/Takes/Source/TakeRecorder/Private/Widgets/STakeRecorderCockpit.h
+++ b/Engine/Plugins/VirtualProduction/Takes/Source/TakeRecorder/Private/Widgets/STakeRecorderCockpit.h
@@ -13,6 +13,7 @@ class ULevelSequence;
 class UTakeMetaData;
 class UTakeRecorder;
 class FUICommandList;
+struct FDigitsTypeInterface;
 
 /**
  * Cockpit UI for defining take meta-data.
@@ -158,4 +159,6 @@ private:
 	FDelegateHandle OnRecordingInitializedHandle, OnRecordingFinishedHandle;
 
 	TSharedPtr<FUICommandList> CommandList;
+
+	TSharedPtr<FDigitsTypeInterface> DigitsTypeInterface;
 };
diff --git a/Engine/Plugins/VirtualProduction/Takes/Source/TakeRecorder/TakeRecorder.Build.cs b/Engine/Plugins/VirtualProduction/Takes/Source/TakeRecorder/TakeRecorder.Build.cs
--- a/Engine/Plugins/VirtualProduction/Takes/Source/TakeRecorder/TakeRecorder.Build.cs
+++ b/Engine/Plugins/VirtualProduction/Takes/Source/TakeRecorder/TakeRecorder.Build.cs
@@ -22,6 +22,7 @@ public class TakeRecorder : ModuleRules
 				"LevelSequence",
 				"LevelSequenceEditor",
 				"MovieScene",
+				"MovieSceneTools",
 				"PropertyEditor",
 				"TakesCore",
 				"TimeManagement",
