Take Recorder: Fix recorded data being offset when not recording from 0. Fixed current sequencer time and take recorder end sections being off by a frame. Fixed all active sections getting their start frame reset when adding new sub sections to record.

#rnx

diff --git a/Engine/Plugins/Animation/LiveLink/Source/LiveLinkSequencer/Private/SequencerRecorderSections/MovieSceneLiveLinkTrackRecorder.cpp b/Engine/Plugins/Animation/LiveLink/Source/LiveLinkSequencer/Private/SequencerRecorderSections/MovieSceneLiveLinkTrackRecorder.cpp
--- a/Engine/Plugins/Animation/LiveLink/Source/LiveLinkSequencer/Private/SequencerRecorderSections/MovieSceneLiveLinkTrackRecorder.cpp
+++ b/Engine/Plugins/Animation/LiveLink/Source/LiveLinkSequencer/Private/SequencerRecorderSections/MovieSceneLiveLinkTrackRecorder.cpp
@@ -170,10 +170,10 @@ void UMovieSceneLiveLinkTrackRecorder::SetSectionStartTimecodeImpl(const FTimeco
 		Parameters.User = GetDefault<UTakeRecorderUserSettings>()->Settings;
 		Parameters.Project = GetDefault<UTakeRecorderProjectSettings>()->Settings;
 
-		FFrameRate TickResolution = MovieSceneSection->GetTypedOuter<UMovieScene>()->GetTickResolution();
-		FFrameRate DisplayRate = MovieSceneSection->GetTypedOuter<UMovieScene>()->GetDisplayRate();
+		FFrameRate TickResolution = MovieScene->GetTickResolution();
+		FFrameRate DisplayRate = MovieScene->GetDisplayRate();
 
-		RecordStartFrame = Parameters.Project.bStartAtCurrentTimecode ? FFrameRate::TransformTime(FFrameTime(InSectionStartTimecode.ToFrameNumber(DisplayRate)), DisplayRate, TickResolution).FloorToFrame() : 0;
+		RecordStartFrame = Parameters.Project.bStartAtCurrentTimecode ? FFrameRate::TransformTime(FFrameTime(InSectionStartTimecode.ToFrameNumber(DisplayRate)), DisplayRate, TickResolution).FloorToFrame() : MovieScene->GetPlaybackRange().GetLowerBoundValue();
 	}
 }
 
diff --git a/Engine/Plugins/VirtualProduction/Takes/Source/TakeRecorder/Private/Recorder/TakeRecorder.cpp b/Engine/Plugins/VirtualProduction/Takes/Source/TakeRecorder/Private/Recorder/TakeRecorder.cpp
--- a/Engine/Plugins/VirtualProduction/Takes/Source/TakeRecorder/Private/Recorder/TakeRecorder.cpp
+++ b/Engine/Plugins/VirtualProduction/Takes/Source/TakeRecorder/Private/Recorder/TakeRecorder.cpp
@@ -472,8 +472,8 @@ bool UTakeRecorder::Initialize(ULevelSequence* LevelSequenceBase, UTakeRecorderS
 		FTimecode CurrentTime = FApp::GetTimecode();
 		UMovieScene* MovieScene = SequenceAsset->GetMovieScene();
 		FFrameRate FrameRate = MovieScene->GetDisplayRate();
-		FFrameNumber ViewRangeStart = Parameters.Project.bStartAtCurrentTimecode ? CurrentTime.ToFrameNumber(FrameRate) : 0;
-		double ViewRangeStartSeconds = FrameRate.AsSeconds(ViewRangeStart);
+		FFrameNumber ViewRangeStart = CurrentTime.ToFrameNumber(FrameRate);
+		double ViewRangeStartSeconds = Parameters.Project.bStartAtCurrentTimecode ? FrameRate.AsSeconds(ViewRangeStart) : MovieScene->GetPlaybackRange().GetLowerBoundValue() / MovieScene->GetTickResolution();
 		FAnimatedRange NewRange(ViewRangeStartSeconds - 0.5f, ViewRangeStartSeconds + (Range.GetUpperBoundValue() - Range.GetLowerBoundValue()) + 0.5f);
 		Sequencer->SetViewRange(NewRange, EViewRangeInterpolation::Immediate);
 		Sequencer->SetClampRange(Sequencer->GetViewRange());
@@ -768,6 +768,7 @@ void UTakeRecorder::Start(const FTimecode& InTimecodeSource)
 
 			// Set infinite playback range when starting recording. Playback range will be clamped to the bounds of the sections at the completion of the recording
 			MovieScene->SetPlaybackRange(TRange<FFrameNumber>(PlaybackStartFrame, TNumericLimits<int32>::Max() - 1), false);
+			Sequencer->SetGlobalTime(PlaybackStartFrame);
 		}
 		Sequencer->SetPlaybackStatus(EMovieScenePlayerStatus::Playing);
 	}
diff --git a/Engine/Plugins/VirtualProduction/Takes/Source/TakeRecorderSources/Private/TakeRecorderMicrophoneAudioSource.cpp b/Engine/Plugins/VirtualProduction/Takes/Source/TakeRecorderSources/Private/TakeRecorderMicrophoneAudioSource.cpp
--- a/Engine/Plugins/VirtualProduction/Takes/Source/TakeRecorderSources/Private/TakeRecorderMicrophoneAudioSource.cpp
+++ b/Engine/Plugins/VirtualProduction/Takes/Source/TakeRecorderSources/Private/TakeRecorderMicrophoneAudioSource.cpp
@@ -175,6 +175,9 @@ void UTakeRecorderMicrophoneAudioSource::StopRecording(class ULevelSequence* InS
 	UMovieScene* MovieScene = InSequence->GetMovieScene();
 	check(CachedAudioTrack.IsValid());
 
+	FFrameRate TickResolution = MovieScene->GetTickResolution();
+	FFrameRate DisplayRate = MovieScene->GetDisplayRate();
+
 	if (bReplaceRecordedAudio)
 	{
 		CachedAudioTrack->RemoveAllAnimationData();
@@ -194,10 +197,7 @@ void UTakeRecorderMicrophoneAudioSource::StopRecording(class ULevelSequence* InS
 
 		UMovieSceneAudioSection* NewAudioSection = NewObject<UMovieSceneAudioSection>(CachedAudioTrack.Get(), UMovieSceneAudioSection::StaticClass());
 
-		FFrameRate TickResolution = CachedAudioTrack->GetTypedOuter<UMovieScene>()->GetTickResolution();
-		FFrameRate DisplayRate = CachedAudioTrack->GetTypedOuter<UMovieScene>()->GetDisplayRate();
-
-		FFrameNumber RecordStartFrame = Parameters.Project.bStartAtCurrentTimecode ? FFrameRate::TransformTime(FFrameTime(TimecodeSource.ToFrameNumber(DisplayRate)), DisplayRate, TickResolution).FloorToFrame() : 0;
+		FFrameNumber RecordStartFrame = Parameters.Project.bStartAtCurrentTimecode ? FFrameRate::TransformTime(FFrameTime(TimecodeSource.ToFrameNumber(DisplayRate)), DisplayRate, TickResolution).FloorToFrame() : MovieScene->GetPlaybackRange().GetLowerBoundValue();
 
 		NewAudioSection->SetSound(RecordedAudio);
 		NewAudioSection->SetRange(TRange<FFrameNumber>(RecordStartFrame, RecordStartFrame + (RecordedAudio->GetDuration() * TickResolution).CeilToFrame()));
diff --git a/Engine/Plugins/VirtualProduction/Takes/Source/TakesCore/Private/TakeRecorderSources.cpp b/Engine/Plugins/VirtualProduction/Takes/Source/TakesCore/Private/TakeRecorderSources.cpp
--- a/Engine/Plugins/VirtualProduction/Takes/Source/TakesCore/Private/TakeRecorderSources.cpp
+++ b/Engine/Plugins/VirtualProduction/Takes/Source/TakesCore/Private/TakeRecorderSources.cpp
@@ -145,7 +145,7 @@ void UTakeRecorderSources::StartRecordingRecursive(TArray<UTakeRecorderSource*>
 				// We initialize the sequence to start at the current time.
 				// We'll have to update these sections each frame as the recording progresses so they appear to get longer like normal
 				// tracks do as we record into them.
-				FFrameNumber RecordStartFrame = bStartAtCurrentTimecode ? FFrameRate::TransformTime(FFrameTime(Timecode.ToFrameNumber(TargetLevelSequenceDisplayRate)), TargetLevelSequenceDisplayRate, TargetLevelSequenceTickResolution).FloorToFrame() : 0;
+				FFrameNumber RecordStartFrame = bStartAtCurrentTimecode ? FFrameRate::TransformTime(FFrameTime(Timecode.ToFrameNumber(TargetLevelSequenceDisplayRate)), TargetLevelSequenceDisplayRate, TargetLevelSequenceTickResolution).FloorToFrame() : InMasterSequence->GetMovieScene()->GetPlaybackRange().GetLowerBoundValue();
 				UMovieSceneSubSection* NewSubSection = SubsceneTrack->AddSequence(TargetSequence, RecordStartFrame, 0);
 				SetSectionStartTimecode(NewSubSection, Timecode, TargetLevelSequenceDisplayRate, TargetLevelSequenceTickResolution);
 
@@ -307,8 +307,8 @@ void UTakeRecorderSources::PreRecordingRecursive(TArray<UTakeRecorderSource*> In
 				// We initialize the sequence to start at zero and be a 0 frame length section as there is no data in the sections yet.
 				// We'll have to update these sections each frame as the recording progresses so they appear to get longer like normal
 				// tracks do as we record into them.
-				FFrameNumber RecordStartTime = FFrameNumber(0);
-				UMovieSceneSubSection* NewSubSection = SubsceneTrack->AddSequence(TargetSequence, RecordStartTime, 0);
+				FFrameNumber RecordStartFrame = bStartAtCurrentTimecode ? FFrameRate::TransformTime(FFrameTime(FApp::GetTimecode().ToFrameNumber(TargetLevelSequenceDisplayRate)), TargetLevelSequenceDisplayRate, TargetLevelSequenceTickResolution).FloorToFrame() : InMasterSequence->GetMovieScene()->GetPlaybackRange().GetLowerBoundValue();
+				UMovieSceneSubSection* NewSubSection = SubsceneTrack->AddSequence(TargetSequence, RecordStartFrame, 0);
 
 				// Section should not be transactional during the recording process
 				NewSubSection->ClearFlags(RF_Transactional);
@@ -426,7 +426,11 @@ void UTakeRecorderSources::StartRecordingTheseSources(const TArray<UTakeRecorder
 			{
 				for (UMovieSceneSubSection* ActiveSubSection : ActiveSubSections)
 				{
-					SetSectionStartTimecode(ActiveSubSection, CurrentTimecode, TargetLevelSequenceDisplayRate, TargetLevelSequenceTickResolution);
+					// Set timecode source and start time if it hasn't been set
+					if (ActiveSubSection->TimecodeSource == FMovieSceneTimecodeSource())
+					{
+						SetSectionStartTimecode(ActiveSubSection, CurrentTimecode, TargetLevelSequenceDisplayRate, TargetLevelSequenceTickResolution);
+					}
 				}
 			}
 			FFrameNumber FrameNumber = QualifiedSequenceTime.ConvertTo(SourceSequence->GetMovieScene()->GetTickResolution()).FloorToFrame();
@@ -439,8 +443,7 @@ void UTakeRecorderSources::StartRecordingTheseSources(const TArray<UTakeRecorder
 
 void UTakeRecorderSources::SetSectionStartTimecode(UMovieSceneSubSection* SubSection, const FTimecode& Timecode, FFrameRate FrameRate, FFrameRate TickResolution)
 {
-	FFrameNumber RecordStartFrame = bStartAtCurrentTimecode ? FFrameRate::TransformTime(FFrameTime(Timecode.ToFrameNumber(FrameRate)), FrameRate, TickResolution).FloorToFrame() : 0;
-
+	FFrameNumber RecordStartFrame = bStartAtCurrentTimecode ? FFrameRate::TransformTime(FFrameTime(Timecode.ToFrameNumber(FrameRate)), FrameRate, TickResolution).FloorToFrame() : GetCurrentRecordingFrameTime().ConvertTo(TickResolution).FloorToFrame();
 	SubSection->TimecodeSource = FMovieSceneTimecodeSource(Timecode);
 
 	// Ensure we're expanded to at least the next frame so that we don't set the start past the end
@@ -497,28 +500,27 @@ FFrameTime UTakeRecorderSources::TickRecording(class ULevelSequence* InSequence,
 	//Time in seconds since recording started. Used when there is no Timecode Sync (e.g. in case it get's lost or dropped).
 	TimeSinceRecordingStarted += DeltaTime;
 
-	const FFrameTime CurrentFrameTimeSinceStart = TargetLevelSequenceTickResolution.AsFrameTime(TimeSinceRecordingStarted);
 
 	// If we're recording into sub-sections we want to update their range every frame so they appear to
 	// animate as their contents are filled. We can't check against the size of all sections (not all
 	// source types have data in their sections until the end) and if you're partially re-recording
 	// a track it would size to the existing content which would skip the animation as well.
 
-	FFrameNumber EndFrame = bStartAtCurrentTimecode ? FFrameRate::TransformTime(FFrameTime(InTimecodeSource.ToFrameNumber(TargetLevelSequenceDisplayRate)), TargetLevelSequenceDisplayRate, TargetLevelSequenceTickResolution).CeilToFrame() : CurrentFrameTimeSinceStart.CeilToFrame();
+	FFrameNumber EndFrame = bStartAtCurrentTimecode ? FFrameRate::TransformTime(FFrameTime(InTimecodeSource.ToFrameNumber(TargetLevelSequenceDisplayRate)), TargetLevelSequenceDisplayRate, TargetLevelSequenceTickResolution).CeilToFrame() : FrameTime.ConvertTo(TargetLevelSequenceTickResolution).CeilToFrame();
 	for (UMovieSceneSubSection* SubSection : ActiveSubSections)
 	{
 		// Subsections will have been created to start at the time that they appeared, so we just need to expand their range to this recording time
 		SubSection->ExpandToFrame(EndFrame);
 	}
-	return CurrentFrameTimeSinceStart;
+	return FrameTime.ConvertTo(TargetLevelSequenceTickResolution);
 }
 
 //We now always just use TimeSinceRecordingStarted instead of possibly using timecode to determine our time since
 //That can give us a higher resolution
 FQualifiedFrameTime UTakeRecorderSources::GetCurrentRecordingFrameTime() const
 {
-	const FFrameNumber StartFrameNumber = bStartAtCurrentTimecode ? StartRecordingTimecodeSource.ToFrameNumber(TargetLevelSequenceDisplayRate) : 0;
-	FFrameTime StartTime = FFrameRate::TransformTime(FFrameTime(StartFrameNumber), TargetLevelSequenceDisplayRate, TargetLevelSequenceTickResolution);
+	const FFrameNumber StartFrameNumber = StartRecordingTimecodeSource.ToFrameNumber(TargetLevelSequenceDisplayRate);
+	FFrameTime StartTime = bStartAtCurrentTimecode ? FFrameRate::TransformTime(FFrameTime(StartFrameNumber), TargetLevelSequenceDisplayRate, TargetLevelSequenceTickResolution) : CachedLevelSequence->GetMovieScene()->GetPlaybackRange().GetLowerBoundValue();
 
 	const FFrameTime CurrentFrameTimeSinceStart = TargetLevelSequenceTickResolution.AsFrameTime(TimeSinceRecordingStarted);
 	FQualifiedFrameTime FrameTime = FQualifiedFrameTime(StartTime + CurrentFrameTimeSinceStart, TargetLevelSequenceTickResolution);
diff --git a/Engine/Source/Runtime/MovieScene/Public/MovieScene.h b/Engine/Source/Runtime/MovieScene/Public/MovieScene.h
--- a/Engine/Source/Runtime/MovieScene/Public/MovieScene.h
+++ b/Engine/Source/Runtime/MovieScene/Public/MovieScene.h
@@ -52,6 +52,15 @@ struct FMovieSceneTimecodeSource
 		, DeltaFrame(FFrameNumber())
 	{}
 
+	FORCEINLINE bool operator==(const FMovieSceneTimecodeSource& Other) const
+	{
+		return Timecode == Other.Timecode && DeltaFrame == Other.DeltaFrame;
+	}
+	FORCEINLINE bool operator!=(const FMovieSceneTimecodeSource& Other) const
+	{
+		return Timecode != Other.Timecode || DeltaFrame != Other.DeltaFrame;
+	}
+
 public:
 
 	/** The global timecode at which this target is based (ie. the timecode at the beginning of the movie scene section when it was recorded) */
