Take Recorder: Add external source details last so that selected sources will be displayed in the take recorder settings first.

diff --git a/Engine/Plugins/VirtualProduction/Takes/Source/TakeRecorder/Private/Widgets/SLevelSequenceTakeEditor.cpp b/Engine/Plugins/VirtualProduction/Takes/Source/TakeRecorder/Private/Widgets/SLevelSequenceTakeEditor.cpp
--- a/Engine/Plugins/VirtualProduction/Takes/Source/TakeRecorder/Private/Widgets/SLevelSequenceTakeEditor.cpp
+++ b/Engine/Plugins/VirtualProduction/Takes/Source/TakeRecorder/Private/Widgets/SLevelSequenceTakeEditor.cpp
@@ -772,25 +772,60 @@ class FRecorderSourceObjectCustomization : public IDetailCustomization
 };
 
 
-void SLevelSequenceTakeEditor::UpdateDetails()
+void SLevelSequenceTakeEditor::AddDetails(const TPair<const UClass*, TArray<UObject*> >& Pair, TArray<FObjectKey>& PreviousClasses)
 {
 	FPropertyEditorModule& PropertyEditorModule = FModuleManager::Get().LoadModuleChecked<FPropertyEditorModule>("PropertyEditor");
 	FDetailsViewArgs DetailsViewArgs(false, false, false, FDetailsViewArgs::HideNameArea, true);
 	DetailsViewArgs.bShowScrollBar = false;
 
-	TSortedMap<const UClass*, TArray<UObject*>> ClassToSources;
+	PreviousClasses.Remove(Pair.Key);
+
+	TSharedPtr<IDetailsView> ExistingDetails = ClassToDetailsView.FindRef(Pair.Key);
+	if (ExistingDetails.IsValid())
+	{
+		ExistingDetails->SetObjects(Pair.Value);
+	}
+	else
+	{
+		TSharedRef<IDetailsView> Details = PropertyEditorModule.CreateDetailView(DetailsViewArgs);
+
+		// Register the custom property layout for all object types to rename the category to the object type
+		// @note: this is registered as a base for all objects on the details panel that
+		// overrides the category name for *all* properties in the object. This makes property categories irrelevant for recorder sources,
+		// And may also interfere with any other detail customizations for sources as a whole if any are added in future (property type customizations will still work fine)
+		// We may want to change this in future but it seems like the neatest way to make top level categories have helpful names.
+		Details->RegisterInstancedCustomPropertyLayout(UTakeRecorderSource::StaticClass(), FOnGetDetailCustomizationInstance::CreateLambda(&MakeShared<FRecorderSourceObjectCustomization>));
+
+		Details->RegisterInstancedCustomPropertyTypeLayout(FName(TEXT("ActorRecorderPropertyMap")), FOnGetPropertyTypeCustomizationInstance::CreateLambda(&MakeShared<FRecorderPropertyMapCustomization >));
+		Details->RegisterInstancedCustomPropertyTypeLayout(FName(TEXT("ActorRecordedProperty")), FOnGetPropertyTypeCustomizationInstance::CreateLambda(&MakeShared<FRecordedPropertyCustomization >));
+		Details->SetObjects(Pair.Value);
+
+		Details->SetEnabled(LevelSequenceAttribute.IsSet() && LevelSequenceAttribute.Get()->FindMetaData<UTakeMetaData>() ? !LevelSequenceAttribute.Get()->FindMetaData<UTakeMetaData>()->Recorded() : true);
+
+		DetailsBox->AddSlot()
+			[
+				Details
+			];
+
+		ClassToDetailsView.Add(Pair.Key, Details);
+	}
+}
+
+void SLevelSequenceTakeEditor::UpdateDetails()
+{
+	TMap<const UClass*, TArray<UObject*>> ExternalClassToSources;
 
-	// External settings objects should be displayed after selected sources, 
-	// so add them to the sorted map first, since the details views will be displayed in reverse order
 	for (TWeakObjectPtr<> WeakExternalObj : ExternalSettingsObjects)
 	{
 		UObject* Object = WeakExternalObj.Get();
 		if (Object)
 		{
-			ClassToSources.FindOrAdd(Object->GetClass()).Add(Object);
+			ExternalClassToSources.FindOrAdd(Object->GetClass()).Add(Object);
 		}
 	}
 
+	TMap<const UClass*, TArray<UObject*>> ClassToSources;
+
 	TArray<UTakeRecorderSource*> SelectedSources;
 	SourcesWidget->GetSelectedSources(SelectedSources);
 
@@ -826,39 +861,14 @@ void SLevelSequenceTakeEditor::UpdateDetails()
 
 	for (auto& Pair : ClassToSources)
 	{
-		PreviousClasses.Remove(Pair.Key);
-
-		TSharedPtr<IDetailsView> ExistingDetails = ClassToDetailsView.FindRef(Pair.Key);
-		if (ExistingDetails.IsValid())
-		{
-			ExistingDetails->SetObjects(Pair.Value);
-		}
-		else
-		{
-			TSharedRef<IDetailsView> Details = PropertyEditorModule.CreateDetailView(DetailsViewArgs);
-
-			// Register the custom property layout for all object types to rename the category to the object type
-			// @note: this is registered as a base for all objects on the details panel that
-			// overrides the category name for *all* properties in the object. This makes property categories irrelevant for recorder sources,
-			// And may also interfere with any other detail customizations for sources as a whole if any are added in future (property type customizations will still work fine)
-			// We may want to change this in future but it seems like the neatest way to make top level categories have helpful names.
-			Details->RegisterInstancedCustomPropertyLayout(UTakeRecorderSource::StaticClass(), FOnGetDetailCustomizationInstance::CreateLambda(&MakeShared<FRecorderSourceObjectCustomization>));
-
-			Details->RegisterInstancedCustomPropertyTypeLayout(FName(TEXT("ActorRecorderPropertyMap")), FOnGetPropertyTypeCustomizationInstance::CreateLambda(&MakeShared<FRecorderPropertyMapCustomization >));
-			Details->RegisterInstancedCustomPropertyTypeLayout(FName(TEXT("ActorRecordedProperty")), FOnGetPropertyTypeCustomizationInstance::CreateLambda(&MakeShared<FRecordedPropertyCustomization >));
-			Details->SetObjects(Pair.Value);
-
-			Details->SetEnabled( LevelSequenceAttribute.IsSet() && LevelSequenceAttribute.Get()->FindMetaData<UTakeMetaData>() ? !LevelSequenceAttribute.Get()->FindMetaData<UTakeMetaData>()->Recorded() : true );
-
-			DetailsBox->AddSlot()
-			[
-				Details
-			];
-
-			ClassToDetailsView.Add(Pair.Key, Details);
-		}
+		AddDetails(Pair, PreviousClasses);
 	}
 
+	for (auto& Pair : ExternalClassToSources)
+	{
+		AddDetails(Pair, PreviousClasses);
+	}
+	
 	for (FObjectKey StaleClass : PreviousClasses)
 	{
 		TSharedPtr<IDetailsView> Details = ClassToDetailsView.FindRef(StaleClass);
diff --git a/Engine/Plugins/VirtualProduction/Takes/Source/TakeRecorder/Private/Widgets/SLevelSequenceTakeEditor.h b/Engine/Plugins/VirtualProduction/Takes/Source/TakeRecorder/Private/Widgets/SLevelSequenceTakeEditor.h
--- a/Engine/Plugins/VirtualProduction/Takes/Source/TakeRecorder/Private/Widgets/SLevelSequenceTakeEditor.h
+++ b/Engine/Plugins/VirtualProduction/Takes/Source/TakeRecorder/Private/Widgets/SLevelSequenceTakeEditor.h
@@ -62,6 +62,8 @@ private:
 	 */
 	void CheckForNewLevelSequence();
 
+	void AddDetails(const TPair<const UClass*, TArray<UObject*> >& Pair, TArray<FObjectKey>& PreviousClasses);
+
 	/**
 	 * Update the details panel for the current selection
 	 */
