Take Recorder: Fix recording spawned particle systems. Particle systems are spawned and attached to the world settings actor. These were not getting picked up when gathering newly created components because they are scene components on the actor and not the root component. Also, fixed an issue where the property map needs to be regenerated when adding the world settings actor as a new actor source during PreRecording

diff --git a/Engine/Plugins/VirtualProduction/Takes/Source/TakeRecorderSources/Private/TakeRecorderActorSource.cpp b/Engine/Plugins/VirtualProduction/Takes/Source/TakeRecorderSources/Private/TakeRecorderActorSource.cpp
--- a/Engine/Plugins/VirtualProduction/Takes/Source/TakeRecorderSources/Private/TakeRecorderActorSource.cpp
+++ b/Engine/Plugins/VirtualProduction/Takes/Source/TakeRecorderSources/Private/TakeRecorderActorSource.cpp
@@ -720,11 +720,10 @@ bool UTakeRecorderActorSource::EnsureObjectTemplateHasComponent(UActorComponent*
 {
 	check(InComponent);
 
-	//If it's native it will be caught below as a component so might as well bail here.
 	//If it's coming from a component that is created from a template defined in the Components section of the Blueprint
 	//it will NOT be found as a Component in AllChildren below but it will exist when created so we also exit here.
 	//So we only do the check if UserConstructionScript or Instance, with the latter may not be needed.. not sure
-	if (InComponent->CreationMethod == EComponentCreationMethod::Native || InComponent->CreationMethod == EComponentCreationMethod::SimpleConstructionScript)
+	if (InComponent->CreationMethod == EComponentCreationMethod::SimpleConstructionScript)
 	{
 		return false;
 	}
@@ -1359,7 +1358,13 @@ void UTakeRecorderActorSource::GetActorComponents(AActor* OnActor, TSet<UActorCo
 
 		for (UActorComponent* ActorComponent : ActorComponents)
 		{
-			if (!ActorComponent->IsA<USceneComponent>())
+			USceneComponent* SceneComponent = Cast<USceneComponent>(ActorComponent);
+
+			// Child of the root component are gathered in GetSceneComponents(). 
+			// Here we gather the rest of the components - either non scene components or 
+			// scene components that are not directly attached to the root component. This 
+			// includes spawned particle systems
+			if (!SceneComponent || !SceneComponent->GetAttachParent())
 			{
 				if (ActorComponent->GetOwner() != Target.Get())
 				{
diff --git a/Engine/Plugins/VirtualProduction/Takes/Source/TakeRecorderSources/Private/TakeRecorderWorldSource.cpp b/Engine/Plugins/VirtualProduction/Takes/Source/TakeRecorderSources/Private/TakeRecorderWorldSource.cpp
--- a/Engine/Plugins/VirtualProduction/Takes/Source/TakeRecorderSources/Private/TakeRecorderWorldSource.cpp
+++ b/Engine/Plugins/VirtualProduction/Takes/Source/TakeRecorderSources/Private/TakeRecorderWorldSource.cpp
@@ -80,6 +80,12 @@ TArray<UTakeRecorderSource*> UTakeRecorderWorldSource::PreRecording(class ULevel
 
 		UTakeRecorderActorSource* ActorSource = NewObject<UTakeRecorderActorSource>(Sources, UTakeRecorderActorSource::StaticClass(), NAME_None, RF_Transactional);
 		ActorSource->Target = WorldSettings;
+
+		// Send a PropertyChangedEvent so the class catches the callback and rebuilds the property map. We can't rely on the Actor rebuilding the map on PreRecording
+		// because that would wipe out any user adjustments from one added natively.
+		FPropertyChangedEvent PropertyChangedEvent(UTakeRecorderActorSource::StaticClass()->FindPropertyByName(GET_MEMBER_NAME_CHECKED(UTakeRecorderActorSource, Target)), EPropertyChangeType::ValueSet);
+		ActorSource->PostEditChangeProperty(PropertyChangedEvent);
+
 		NewSources.Add(ActorSource);
 		WorldSource = ActorSource;
 	}
diff --git a/Engine/Plugins/VirtualProduction/Takes/Source/TakesCore/Private/TakesUtils.cpp b/Engine/Plugins/VirtualProduction/Takes/Source/TakesCore/Private/TakesUtils.cpp
--- a/Engine/Plugins/VirtualProduction/Takes/Source/TakesCore/Private/TakesUtils.cpp
+++ b/Engine/Plugins/VirtualProduction/Takes/Source/TakesCore/Private/TakesUtils.cpp
@@ -43,13 +43,16 @@ void ClampPlaybackRangeToEncompassAllSections(UMovieScene* InMovieScene, bool bU
 	for (UMovieSceneSection* Section : MovieSceneSections)
 	{
 		TRange<FFrameNumber> SectionRange = Section->GetRange();
-		if (!PlayRange.IsSet())
+		if (SectionRange.GetLowerBound().IsClosed() && SectionRange.GetUpperBound().IsClosed())
 		{
-			PlayRange = SectionRange;
-		}
-		else if (SectionRange.GetLowerBound().IsClosed() && SectionRange.GetUpperBound().IsClosed())
-		{
-			PlayRange = TRange<FFrameNumber>::Hull(PlayRange.GetValue(), SectionRange);
+			if (!PlayRange.IsSet())
+			{
+				PlayRange = SectionRange;
+			}
+			else
+			{
+				PlayRange = TRange<FFrameNumber>::Hull(PlayRange.GetValue(), SectionRange);
+			}
 		}
 	}
 
