Take Recorder: Fix incorrect check for whether the timecode has been set on sections.

diff --git a/Engine/Plugins/VirtualProduction/Takes/Source/TakesCore/Private/TakeRecorderSources.cpp b/Engine/Plugins/VirtualProduction/Takes/Source/TakesCore/Private/TakeRecorderSources.cpp
--- a/Engine/Plugins/VirtualProduction/Takes/Source/TakesCore/Private/TakeRecorderSources.cpp
+++ b/Engine/Plugins/VirtualProduction/Takes/Source/TakesCore/Private/TakeRecorderSources.cpp
@@ -427,7 +427,7 @@ void UTakeRecorderSources::StartRecordingTheseSources(const TArray<UTakeRecorder
 				for (UMovieSceneSubSection* ActiveSubSection : ActiveSubSections)
 				{
 					// Set timecode source and start time if it hasn't been set
-					if (ActiveSubSection->TimecodeSource == FMovieSceneTimecodeSource())
+					if (ActiveSubSection->TimecodeSource.Timecode == FTimecode())
 					{
 						SetSectionStartTimecode(ActiveSubSection, CurrentTimecode, TargetLevelSequenceDisplayRate, TargetLevelSequenceTickResolution);
 					}
