Take Recorder: Fix enum track recorder not setting the enum being recorded which caused a crash when recording the VisibilityAnimTickOption property

diff --git a/Engine/Plugins/VirtualProduction/Takes/Source/TakeTrackRecorders/Private/MovieScenePropertyTrackRecorder.cpp b/Engine/Plugins/VirtualProduction/Takes/Source/TakeTrackRecorders/Private/MovieScenePropertyTrackRecorder.cpp
--- a/Engine/Plugins/VirtualProduction/Takes/Source/TakeTrackRecorders/Private/MovieScenePropertyTrackRecorder.cpp
+++ b/Engine/Plugins/VirtualProduction/Takes/Source/TakeTrackRecorders/Private/MovieScenePropertyTrackRecorder.cpp
@@ -68,7 +68,7 @@ UMovieSceneTrackRecorder* FMovieScenePropertyTrackRecorderFactory::CreateTrackRe
 		TrackRecorder->PropertyRecorder = MakeShareable(new FMovieSceneTrackPropertyRecorder<uint8>(Binding));
 		break;
 	case ESerializedPropertyType::EnumType:
-		TrackRecorder->PropertyRecorder = MakeShareable(new FMovieSceneTrackPropertyRecorderEnum(Binding));
+		TrackRecorder->PropertyRecorder = MakeShareable(new FMovieSceneTrackPropertyRecorder<uint8>(Binding));
 		break;
 	case ESerializedPropertyType::IntegerType:
 		TrackRecorder->PropertyRecorder = MakeShareable(new FMovieSceneTrackPropertyRecorder<int32>(Binding));
@@ -106,7 +106,7 @@ void UMovieScenePropertyTrackRecorder::CreateTrackImpl()
  		}
 		else if (Property->IsA<FEnumProperty>())
 		{
-			PropertyRecorder = MakeShareable(new FMovieSceneTrackPropertyRecorderEnum(Binding));
+			PropertyRecorder = MakeShareable(new FMovieSceneTrackPropertyRecorder<uint8>(Binding));
 		}
 		else if (Property->IsA<FIntProperty>())
 		{
diff --git a/Engine/Plugins/VirtualProduction/Takes/Source/TakeTrackRecorders/Private/MovieSceneTrackPropertyRecorder.cpp b/Engine/Plugins/VirtualProduction/Takes/Source/TakeTrackRecorders/Private/MovieSceneTrackPropertyRecorder.cpp
--- a/Engine/Plugins/VirtualProduction/Takes/Source/TakeTrackRecorders/Private/MovieSceneTrackPropertyRecorder.cpp
+++ b/Engine/Plugins/VirtualProduction/Takes/Source/TakeTrackRecorders/Private/MovieSceneTrackPropertyRecorder.cpp
@@ -75,6 +75,11 @@ UMovieSceneSection* FMovieSceneTrackPropertyRecorder<bool>::AddSection(const FSt
 	return nullptr;
 }
 
+template <>
+void FMovieSceneTrackPropertyRecorder<bool>::PostCreate(IMovieSceneTrackRecorderHost* InRecordingHost, UObject* InObjectToRecord, class UMovieScene* InMovieScene, const FGuid& InGuid, bool bOpenSerializer)
+{
+}
+
 template <>
 void FMovieSceneTrackPropertyRecorder<bool>::AddKeyToSection(UMovieSceneSection* InSection, const FPropertyKey<bool>& InKey)
 {
@@ -230,6 +235,27 @@ UMovieSceneSection* FMovieSceneTrackPropertyRecorder<uint8>::AddSection(const FS
 	return nullptr;
 }
 
+template <>
+void FMovieSceneTrackPropertyRecorder<uint8>::PostCreate(IMovieSceneTrackRecorderHost* InRecordingHost, UObject* InObjectToRecord, class UMovieScene* InMovieScene, const FGuid& InGuid, bool bOpenSerializer)
+{
+	UMovieSceneByteTrack* Track = Cast<UMovieSceneByteTrack>(MovieSceneSection->GetOuter());
+	if (Track && !Track->GetEnum())
+	{
+		FProperty* PropertyToRecord = Binding.GetProperty(*InObjectToRecord);
+		if (PropertyToRecord)
+		{
+			if (FEnumProperty* EnumProperty = CastField<FEnumProperty>(PropertyToRecord))
+			{
+				Track->SetEnum(EnumProperty->GetEnum());
+			}
+			else if (FByteProperty* ByteProperty = CastField<FByteProperty>(PropertyToRecord))
+			{
+				Track->SetEnum(ByteProperty->Enum);
+			}
+		}
+	}
+}
+
 template <>
 void FMovieSceneTrackPropertyRecorder<uint8>::AddKeyToSection(UMovieSceneSection* InSection, const FPropertyKey<uint8>& InKey)
 {
@@ -336,184 +362,6 @@ bool FMovieSceneTrackPropertyRecorder<uint8>::LoadRecordedFile(const FString& Fi
 	return false;
 }
 
-bool FMovieSceneTrackPropertyRecorderEnum::ShouldAddNewKey(const int64& InNewValue) const
-{
-	return InNewValue != PreviousValue;
-}
-
-UMovieSceneSection* FMovieSceneTrackPropertyRecorderEnum::AddSection(const FString& TrackDisplayName, UMovieScene* InMovieScene, const FGuid& InGuid, bool bSetDefault)
-{
-	FName TrackName = *Binding.GetPropertyPath();
-	UMovieSceneEnumTrack* Track = InMovieScene->FindTrack<UMovieSceneEnumTrack>(InGuid, TrackName);
-	if (!Track)
-	{
-		Track = InMovieScene->AddTrack<UMovieSceneEnumTrack>(InGuid);
-	}
-	else
-	{
-		Track->RemoveAllAnimationData();
-	}
-
-	if (Track)
-	{
-		Track->SetPropertyNameAndPath(Binding.GetPropertyName(), Binding.GetPropertyPath());
-
-		UMovieSceneEnumSection* Section = Cast<UMovieSceneEnumSection>(Track->CreateNewSection());
-
-		// We only set the track defaults when we're not loading from a serialized recording. Serialized recordings don't store channel defaults but will always store a
-		// key on the first frame which will accomplish the same.
-		if (bSetDefault)
-		{
-			FMovieSceneByteChannel* Channel = Section->GetChannelProxy().GetChannel<FMovieSceneByteChannel>(0);
-			Channel->SetDefault(PreviousValue);
-		}
-
-		return Section;
-	}
-
-	return nullptr;
-}
-
-void FMovieSceneTrackPropertyRecorderEnum::AddKeyToSection(UMovieSceneSection* InSection, const FPropertyKey<int64>& InKey)
-{
-	InSection->GetChannelProxy().GetChannel<FMovieSceneByteChannel>(0)->GetData().AddKey(InKey.Time, InKey.Value);
-}
-
-void FMovieSceneTrackPropertyRecorderEnum::ReduceKeys(UMovieSceneSection* InSection)
-{
-	// Reduce keys intentionally left blank
-}
-
-void FMovieSceneTrackPropertyRecorderEnum::RemoveRedundantTracks(UMovieSceneSection* InSection, UObject* InObjectToRecord)
-{
-	if (!InObjectToRecord || !InSection)
-	{
-		return;
-	}
-
-	FTrackRecorderSettings TrackRecorderSettings = OwningTakeRecorderSource->GetTrackRecorderSettings();
-
-	FMovieSceneByteChannel* Channel = MovieSceneSection.Get()->GetChannelProxy().GetChannel<FMovieSceneByteChannel>(0);
-	if (!Channel)
-	{
-		return;
-	}
-
-	// If less than 1 key, remove the key and ensure the default value is the same
-	if (Channel->GetData().GetValues().Num() > 1)
-	{
-		return;
-	}
-
-	int64 DefaultValue = Channel->GetData().GetValues().Num() == 1 ? Channel->GetData().GetValues()[0] : Channel->GetDefault().GetValue();
-	Channel->GetData().Reset();
-	Channel->SetDefault(DefaultValue);
-
-	// The section can be removed if this is a spawnable since the spawnable template should have the same default values
-	bool bRemoveSection = true;
-
-	// If recording to a possessable, this section can only be removed if the CDO value is the same and it's not on the whitelist of default property tracks
-	if (TrackRecorderSettings.bRecordToPossessable)
-	{
-		bRemoveSection = false;
-
-		UObject* DefaultObject = InObjectToRecord->GetClass()->GetDefaultObject();
-		if (DefaultObject && Binding.GetCurrentValue<int64>(*DefaultObject) == DefaultValue)
-		{
-			bRemoveSection = true;
-		}
-
-		if (bRemoveSection && FTrackRecorderSettings::IsDefaultPropertyTrack(InObjectToRecord, Binding.GetPropertyPath(), TrackRecorderSettings.DefaultTracks))
-		{
-			bRemoveSection = false;
-		}
-	}
-	
-	if (!bRemoveSection && FTrackRecorderSettings::IsExcludePropertyTrack(InObjectToRecord, Binding.GetPropertyPath(), TrackRecorderSettings.DefaultTracks))
-	{
-		bRemoveSection = true;
-	}
-
-	if (bRemoveSection)
-	{
-		UMovieSceneTrack* MovieSceneTrack = CastChecked<UMovieSceneTrack>(MovieSceneSection->GetOuter());
-		UMovieScene* MovieScene = CastChecked<UMovieScene>(MovieSceneTrack->GetOuter());
-
-		UE_LOG(LogTakesCore, Log, TEXT("Removed unused track (%s) for (%s)"), *MovieSceneTrack->GetTrackName().ToString(), *InObjectToRecord->GetName());
-
-		MovieSceneTrack->RemoveSection(*MovieSceneSection.Get());
-		MovieScene->RemoveTrack(*MovieSceneTrack);
-	}
-}
-
-bool FMovieSceneTrackPropertyRecorderEnum::OpenSerializer(const FString& InObjectName, const FName& InPropertyName, const FString& InTrackDisplayName, const FGuid& InGuid)
-{
-	FName SerializedType("Property");
-	FFrameRate   TickResolution = MovieSceneSection->GetTypedOuter<UMovieScene>()->GetTickResolution();
-	FPropertyFileHeader Header(TickResolution, SerializedType, InGuid);
-	Header.PropertyName = InPropertyName;
-	Header.TrackDisplayName = InTrackDisplayName;
-	Header.PropertyType = ESerializedPropertyType::EnumType;
-	FText Error;
-	FString FileName = FString::Printf(TEXT("%s_%s_%s"), *(SerializedType.ToString()), *InObjectName, *(InPropertyName.ToString()));
-
-	if (!Serializer.OpenForWrite(FileName, Header, Error))
-	{
-		UE_LOG(PropertySerialization, Warning, TEXT("Error Opening Property File: Object '%s' Property '%s' Error: '%s'"), *InObjectName, *InPropertyName.ToString(), *Error.ToString());
-		return false;
-	}
-	return true;
-}
-bool FMovieSceneTrackPropertyRecorderEnum::LoadRecordedFile(const FString& FileName, UMovieScene *InMovieScene, TMap<FGuid, AActor*>& ActorGuidToActorMap, TFunction<void()> InCompletionCallback)
-{
-	bool bFileExists = Serializer.DoesFileExist(FileName);
-	if (bFileExists)
-	{
-		FText Error;
-		FPropertyFileHeader Header;
-
-		if (Serializer.OpenForRead(FileName, Header, Error))
-		{
-			MovieSceneSection = AddSection(Header.TrackDisplayName, InMovieScene, Header.Guid, false);
-
-			Serializer.GetDataRanges([this, InMovieScene, FileName, Header, InCompletionCallback](uint64 InMinFrameId, uint64 InMaxFrameId)
-			{
-				auto OnReadComplete = [this, InMovieScene, Header, InCompletionCallback]()
-				{
-					TArray<FPropertySerializedEnumFrame> &InFrames = Serializer.ResultData;
-					if (InFrames.Num() > 0)
-					{
-						FFrameRate InFrameRate = Header.TickResolution;
-						for (const FPropertySerializedEnumFrame& SerializedFrame : InFrames)
-						{
-							const FPropertySerializedEnum& Frame = SerializedFrame.Frame;
-							FFrameRate   TickResolution = MovieSceneSection->GetTypedOuter<UMovieScene>()->GetTickResolution();
-							FFrameTime FrameTime = FFrameRate::TransformTime(Frame.Time, InFrameRate, TickResolution);
-							FFrameNumber CurrentFrame = FrameTime.FrameNumber;
-							FPropertyKey<int64> Key;
-							Key.Time = CurrentFrame;
-							Key.Value = Frame.Value;
-							AddKeyToSection(MovieSceneSection.Get(), Key);
-							MovieSceneSection->ExpandToFrame(CurrentFrame);
-						}
-					}
-					Serializer.Close();
-					InCompletionCallback();
-				}; //callback
-
-				Serializer.ReadFramesAtFrameRange(InMinFrameId, InMaxFrameId, OnReadComplete);
-
-			});
-			return true;
-		}
-		else
-		{
-			Serializer.Close();
-		}
-	}
-	return false;
-}
-
 template <>
 bool FMovieSceneTrackPropertyRecorder<float>::ShouldAddNewKey(const float& InNewValue) const
 {
@@ -557,6 +405,11 @@ UMovieSceneSection* FMovieSceneTrackPropertyRecorder<float>::AddSection(const FS
 	return nullptr;
 }
 
+template <>
+void FMovieSceneTrackPropertyRecorder<float>::PostCreate(IMovieSceneTrackRecorderHost* InRecordingHost, UObject* InObjectToRecord, class UMovieScene* InMovieScene, const FGuid& InGuid, bool bOpenSerializer)
+{
+}
+
 template <>
 void FMovieSceneTrackPropertyRecorder<float>::AddKeyToSection(UMovieSceneSection* InSection, const FPropertyKey<float>& InKey)
 {
@@ -710,6 +563,11 @@ UMovieSceneSection* FMovieSceneTrackPropertyRecorder<FColor>::AddSection(const F
 	return nullptr;
 }
 
+template <>
+void FMovieSceneTrackPropertyRecorder<FColor>::PostCreate(IMovieSceneTrackRecorderHost* InRecordingHost, UObject* InObjectToRecord, class UMovieScene* InMovieScene, const FGuid& InGuid, bool bOpenSerializer)
+{
+}
+
 template <>
 void FMovieSceneTrackPropertyRecorder<FColor>::AddKeyToSection(UMovieSceneSection* InSection, const FPropertyKey<FColor>& InKey)
 {
@@ -907,6 +765,11 @@ UMovieSceneSection* FMovieSceneTrackPropertyRecorder<FVector>::AddSection(const
 	return nullptr;
 }
 
+template <>
+void FMovieSceneTrackPropertyRecorder<FVector>::PostCreate(IMovieSceneTrackRecorderHost* InRecordingHost, UObject* InObjectToRecord, class UMovieScene* InMovieScene, const FGuid& InGuid, bool bOpenSerializer)
+{
+}
+
 template <>
 void FMovieSceneTrackPropertyRecorder<FVector>::AddKeyToSection(UMovieSceneSection* InSection, const FPropertyKey<FVector>& InKey)
 {
@@ -1083,6 +946,11 @@ UMovieSceneSection* FMovieSceneTrackPropertyRecorder<int32>::AddSection(const FS
 	return nullptr;
 }
 
+template <>
+void FMovieSceneTrackPropertyRecorder<int32>::PostCreate(IMovieSceneTrackRecorderHost* InRecordingHost, UObject* InObjectToRecord, class UMovieScene* InMovieScene, const FGuid& InGuid, bool bOpenSerializer)
+{
+}
+
 template <>
 void FMovieSceneTrackPropertyRecorder<int32>::AddKeyToSection(UMovieSceneSection* InSection, const FPropertyKey<int32>& InKey)
 {
@@ -1244,6 +1112,11 @@ UMovieSceneSection* FMovieSceneTrackPropertyRecorder<FString>::AddSection(const
 	return nullptr;
 }
 
+template <>
+void FMovieSceneTrackPropertyRecorder<FString>::PostCreate(IMovieSceneTrackRecorderHost* InRecordingHost, UObject* InObjectToRecord, class UMovieScene* InMovieScene, const FGuid& InGuid, bool bOpenSerializer)
+{
+}
+
 template <>
 void FMovieSceneTrackPropertyRecorder<FString>::AddKeyToSection(UMovieSceneSection* InSection, const FPropertyKey<FString>& InKey)
 {
diff --git a/Engine/Plugins/VirtualProduction/Takes/Source/TakeTrackRecorders/Public/TrackRecorders/MovieSceneTrackPropertyRecorder.h b/Engine/Plugins/VirtualProduction/Takes/Source/TakeTrackRecorders/Public/TrackRecorders/MovieSceneTrackPropertyRecorder.h
--- a/Engine/Plugins/VirtualProduction/Takes/Source/TakeTrackRecorders/Public/TrackRecorders/MovieSceneTrackPropertyRecorder.h
+++ b/Engine/Plugins/VirtualProduction/Takes/Source/TakeTrackRecorders/Public/TrackRecorders/MovieSceneTrackPropertyRecorder.h
@@ -82,6 +82,8 @@ public:
 			{	
 				OpenSerializer(InObjectToRecord->GetName(), Binding.GetPropertyName(), TrackDisplayName, InGuid);
 			}
+
+			PostCreate(InRecordingHost, InObjectToRecord, InMovieScene, InGuid, bOpenSerializer);
 		}
 	}
 
@@ -276,6 +278,9 @@ private:
 	/** Helper function, specialized by type, used to add an appropriate section to the movie scene */
 	class UMovieSceneSection* AddSection(const FString& TrackDisplayName, class UMovieScene* InMovieScene, const FGuid& InGuid, bool bSetDefault = true);
 
+	/** Helper function, specialized by type, after create */
+	void PostCreate(IMovieSceneTrackRecorderHost* InRecordingHost, UObject* InObjectToRecord, class UMovieScene* InMovieScene, const FGuid& InGuid, bool bOpenSerializer = false);
+
 	/** Helper function, specialized by type, used to add keys to the movie scene section at Finalize() time */
 	void AddKeyToSection(UMovieSceneSection* InSection, const FPropertyKey<PropertyType>& InKey);
 
@@ -310,179 +315,3 @@ private:
 	/** Serializer*/
 	TMovieSceneSerializer<FPropertyFileHeader, FSerializedProperty<PropertyType>> Serializer;
 };
-
-/** Recorder for a simple property of type enum */
-class  FMovieSceneTrackPropertyRecorderEnum : public IMovieSceneTrackPropertyRecorder
-{
-public:
-	FMovieSceneTrackPropertyRecorderEnum(const FTrackInstancePropertyBindings& InBinding)
-		: Binding(InBinding)
-	{
-	}
-	virtual ~FMovieSceneTrackPropertyRecorderEnum() { }
-
-	virtual void Create(IMovieSceneTrackRecorderHost* InRecordingHost, UObject* InObjectToRecord, class UMovieScene* InMovieScene, const FGuid& InGuid, bool bOpenSerializer) override
-	{
-		OwningTakeRecorderSource = InRecordingHost;
-		if (InObjectToRecord)
-		{
-			PreviousValue = Binding.GetCurrentValueForEnum(*InObjectToRecord);
-		}
-
-		if (!InObjectToRecord)
-		{
-			MovieSceneSection = nullptr;
-		}
-		else
-		{
-			FString TrackDisplayName = *Binding.GetProperty(*InObjectToRecord)->GetDisplayNameText().ToString();
-			MovieSceneSection = AddSection(TrackDisplayName, InMovieScene, InGuid, bOpenSerializer);
-			if (bOpenSerializer)
-			{
-				if (OpenSerializer(InObjectToRecord->GetName(), Binding.GetPropertyName(), TrackDisplayName, InGuid))
-				{
-					// FSerializedProperty<int64> Frame;
-					// FFrameRate   TickResolution = MovieSceneSection->GetTypedOuter<UMovieScene>()->GetTickResolution();
-					// FFrameNumber CurrentFrame = (InTime * TickResolution).FloorToFrame();
-					// Frame.Time = CurrentFrame;
-					// Frame.Value = PreviousValue;
-					// Serializer.WriteFrameData(Serializer.FramesWritten, Frame);
-				}
-			}
-		}
-	}
-
-	virtual void SetSectionStartTimecode(const FTimecode& InSectionStartTimecode, const FFrameNumber& InSectionFirstFrame) override
-	{
-		if (!MovieSceneSection.IsValid())
-		{
-			return;
-		}
-
-		MovieSceneSection->TimecodeSource = FMovieSceneTimecodeSource(InSectionStartTimecode);
-		MovieSceneSection->ExpandToFrame(InSectionFirstFrame + FFrameNumber(1));
-		MovieSceneSection->SetStartFrame(TRangeBound<FFrameNumber>::Inclusive(InSectionFirstFrame));
-	}
-
-	virtual void Record(UObject* InObjectToRecord, const FQualifiedFrameTime& InCurrentTime) override
-	{
-		if (!MovieSceneSection.IsValid())
-		{
-			return;
-		}
-
-		if (InObjectToRecord != nullptr)
-		{
-			FFrameRate   TickResolution  = MovieSceneSection->GetTypedOuter<UMovieScene>()->GetTickResolution();
-			FFrameNumber CurrentFrame    = InCurrentTime.ConvertTo(TickResolution).FloorToFrame();
-
-			MovieSceneSection->SetEndFrame(CurrentFrame);
-
-			int64 NewValue = Binding.GetCurrentValueForEnum(*InObjectToRecord);
-			if (ShouldAddNewKey(NewValue))
-			{
-				if (PreviousFrame.IsSet())
-				{
-					FPropertyKey<int64> Key;
-					Key.Time = PreviousFrame.GetValue();
-					Key.Value = PreviousValue;
-					Keys.Add(Key);
-					FSerializedProperty<int64> Frame;
-					Frame.Time = Key.Time;
-					Frame.Value = Key.Value;
-					Serializer.WriteFrameData(Serializer.FramesWritten, Frame);
-				}
-
-				FPropertyKey<int64> Key;
-				Key.Time = CurrentFrame;
-				Key.Value = NewValue;
-				Keys.Add(Key);
-				FSerializedProperty<int64> Frame;
-				Frame.Time = Key.Time;
-				Frame.Value = Key.Value;
-				Serializer.WriteFrameData(Serializer.FramesWritten, Frame);
-
-				PreviousValue = NewValue;
-				PreviousFrame.Reset();
-			}
-			else
-			{
-				PreviousFrame = CurrentFrame;
-			}
-		}
-	}
-
-	virtual void Finalize(UObject* InObjectToRecord) override
-	{
-		if (!MovieSceneSection.IsValid())
-		{
-			return;
-		}
-
-		for (const FPropertyKey<int64>& Key : Keys)
-		{
-			AddKeyToSection(MovieSceneSection.Get(), Key);
-		}
-
-		FTrackRecorderSettings TrackRecorderSettings = OwningTakeRecorderSource->GetTrackRecorderSettings();
-
-		if (TrackRecorderSettings.bReduceKeys)
-		{
-			ReduceKeys(MovieSceneSection.Get());
-		}
-
-		if (TrackRecorderSettings.bRemoveRedundantTracks)
-		{
-			RemoveRedundantTracks(MovieSceneSection.Get(), InObjectToRecord);
-		}
-
-		Serializer.Close();
-
-	}
-
-	virtual void SetSavedRecordingDirectory(const FString& InDirectory) override
-	{
-		Serializer.SetLocalCaptureDir(InDirectory);
-	}
-	virtual bool LoadRecordedFile(const FString& InFileName, UMovieScene *InMovieScene, TMap<FGuid, AActor*>& ActorGuidToActorMap,  TFunction<void()> InCompletionCallback) override;
-
-private:
-	/** 
-	 * Helper function, specialized by type, used to check if we do capture-time filtering of keys based
-	 * on previous values
-	 */
-	bool ShouldAddNewKey(const int64& InNewValue) const;
-
-	/** Helper function, specialized by type, used to add an appropriate section to the movie scene */
-	class UMovieSceneSection* AddSection(const FString& TrackDisplayName, class UMovieScene* InMovieScene, const FGuid& InGuid, bool bSetDefault = true);
-
-	/** Helper function, specialized by type, used to add keys to the movie scene section at Finalize() time */
-	void AddKeyToSection(UMovieSceneSection* InSection, const FPropertyKey<int64>& InKey);
-
-	/** Helper function, specialized by type, used to reduce keys */
-	void ReduceKeys(UMovieSceneSection* InSection);
-
-	/** Helper function, used to remove redundant tracks */
-	void RemoveRedundantTracks(UMovieSceneSection* InSection, UObject* InObjectToRecord);
-
-	/** Open the Serializer of the right PropertyType*/
-	bool OpenSerializer(const FString& InObjectName, const FName& InPropertyName, const FString& InTrackDisplayName, const FGuid& InGuid);
-
-
-private:
-	/** Binding for this property */
-	FTrackInstancePropertyBindings Binding;
-
-	/** The keys that are being recorded */
-	TArray<FPropertyKey<int64>> Keys;
-
-	/** Section we are recording */
-	TWeakObjectPtr<class UMovieSceneSection> MovieSceneSection;
-
-	/** Previous value we use to establish whether we should key */
-	int64 PreviousValue;
-	TOptional<FFrameNumber> PreviousFrame;
-
-	/** Serializer */
-	FPropertySerializerEnum Serializer;
-};
