Take Recorder: Disallow recording of any non placeable actors

diff --git a/Engine/Plugins/VirtualProduction/Takes/Source/TakeRecorderSources/Private/TakeRecorderSourcesModule.cpp b/Engine/Plugins/VirtualProduction/Takes/Source/TakeRecorderSources/Private/TakeRecorderSourcesModule.cpp
--- a/Engine/Plugins/VirtualProduction/Takes/Source/TakeRecorderSources/Private/TakeRecorderSourcesModule.cpp
+++ b/Engine/Plugins/VirtualProduction/Takes/Source/TakeRecorderSources/Private/TakeRecorderSourcesModule.cpp
@@ -156,7 +156,7 @@ namespace
 			{
 				for (AActor* Actor : Level->Actors)
 				{
-					if (Actor && Actor->IsA(Class) && !Actor->IsA(ALevelScriptActor::StaticClass()) && !Actor->IsA(ALevelSequenceActor::StaticClass()))
+					if (Actor && Actor->IsA(Class) && !Actor->IsA(ALevelScriptActor::StaticClass()) && !Actor->IsA(ALevelSequenceActor::StaticClass()) && !Actor->GetClass()->HasAnyClassFlags(CLASS_NotPlaceable))
 					{
 						OutActors.AddUnique(Actor);
 					}
