Take Recorder: Refactor to cull out valid actors for recording - drag and drop, outliner, and command line. This makes it so that you can't record landscape actors because they are not placeable

diff --git a/Engine/Plugins/VirtualProduction/Takes/Source/TakeRecorderSources/Private/TakeRecorderActorSource.cpp b/Engine/Plugins/VirtualProduction/Takes/Source/TakeRecorderSources/Private/TakeRecorderActorSource.cpp
--- a/Engine/Plugins/VirtualProduction/Takes/Source/TakeRecorderSources/Private/TakeRecorderActorSource.cpp
+++ b/Engine/Plugins/VirtualProduction/Takes/Source/TakeRecorderSources/Private/TakeRecorderActorSource.cpp
@@ -62,6 +62,12 @@ UTakeRecorderSource* UTakeRecorderActorSource::AddSourceForActor(AActor* InActor
 		return nullptr;
 	}
 
+	if (!TakeRecorderSourcesUtils::IsActorRecordable(InActor))
+	{
+		FFrame::KismetExecutionMessage(*FString::Printf(TEXT("The Actor: %s is not recordable."), *InActor->GetPathName()), ELogVerbosity::Error);
+		return nullptr;
+	}
+
 	//Look through our sources and see if one actor matches the incoming one either from editor or PIE world.
 	{
 		//Cache  InputActor comparison data
diff --git a/Engine/Plugins/VirtualProduction/Takes/Source/TakeRecorderSources/Private/TakeRecorderSourcesModule.cpp b/Engine/Plugins/VirtualProduction/Takes/Source/TakeRecorderSources/Private/TakeRecorderSourcesModule.cpp
--- a/Engine/Plugins/VirtualProduction/Takes/Source/TakeRecorderSources/Private/TakeRecorderSourcesModule.cpp
+++ b/Engine/Plugins/VirtualProduction/Takes/Source/TakeRecorderSources/Private/TakeRecorderSourcesModule.cpp
@@ -13,6 +13,7 @@
 #include "TakeRecorderActorSource.h"
 #include "TakeRecorderSources.h"
 #include "TakeRecorderSettings.h"
+#include "TakeRecorderSourcesUtils.h"
 #include "Features/IModularFeatures.h"
 #include "DragAndDrop/ActorDragDropOp.h"
 #include "DragAndDrop/FolderDragDropOp.h"
@@ -178,14 +179,52 @@ struct FActorTakeRecorderDropHandler : ITakeRecorderDropHandler
 	{
 		bool bCanHandle = false;
 		if (InOperation)
-		{
+		{		
+			TSharedPtr<FActorDragDropOp>  ActorDrag = nullptr;
+			TSharedPtr<FFolderDragDropOp> FolderDrag = nullptr;
+
+			FDragDropOperation* OperationPtr = InOperation.Get();
+			if (!OperationPtr)
+			{
+				return false;
+			}
 			if (InOperation->IsOfType<FActorDragDropOp>())
 			{
-				bCanHandle |= InOperation->CastTo<FActorDragDropOp>()->Actors.Num() > 0;
+				ActorDrag = OperationPtr->CastTo<FActorDragDropOp>();
 			}
 			if (InOperation->IsOfType<FFolderDragDropOp>())
 			{
-				bCanHandle |= true;
+				FolderDrag = OperationPtr->CastTo<FFolderDragDropOp>();
+			}
+
+			if (ActorDrag)
+			{
+				for (TWeakObjectPtr<AActor> WeakActor : ActorDrag->Actors)
+				{
+					if (AActor* Actor = WeakActor.Get())
+					{
+						if (TakeRecorderSourcesUtils::IsActorRecordable(Actor))
+						{
+							bCanHandle = true;
+							break;
+						}
+					}
+				}
+			}
+
+			if (FolderDrag && !bCanHandle)
+			{
+				TArray<AActor*> FolderActors;
+				FActorFolders::GetActorsFromFolders(*GWorld, FolderDrag->Folders, FolderActors);
+
+				for (AActor* ActorInFolder : FolderActors)
+				{
+					if (TakeRecorderSourcesUtils::IsActorRecordable(ActorInFolder))
+					{
+						bCanHandle = true;
+						break;
+					}
+				}
 			}
 		}
 
@@ -220,14 +259,26 @@ struct FActorTakeRecorderDropHandler : ITakeRecorderDropHandler
 			{
 				if (AActor* Actor = WeakActor.Get())
 				{
-					DraggedActors.Add(Actor);
+					if (TakeRecorderSourcesUtils::IsActorRecordable(Actor))
+					{
+						DraggedActors.Add(Actor);
+					}
 				}
 			}
 		}
 
 		if (FolderDrag)
 		{
-			FActorFolders::GetActorsFromFolders(*GWorld, FolderDrag->Folders, DraggedActors);
+			TArray<AActor*> FolderActors;
+			FActorFolders::GetActorsFromFolders(*GWorld, FolderDrag->Folders, FolderActors);
+
+			for (AActor* ActorInFolder : FolderActors)
+			{
+				if (TakeRecorderSourcesUtils::IsActorRecordable(ActorInFolder))
+				{
+					DraggedActors.Add(ActorInFolder);
+				}
+			}
 		}
 
 		TArray<AActor*> ExistingActors;
@@ -337,7 +388,7 @@ public:
 
 		auto OutlinerFilterPredicate = [InExistingActors = MoveTemp(ExistingActors)](const AActor* InActor)
 		{
-			return !InExistingActors.Contains(InActor);
+			return !InExistingActors.Contains(InActor) && TakeRecorderSourcesUtils::IsActorRecordable(InActor);
 		};
 
 		// Set up a menu entry to add the selected actor(s) to the sequencer
diff --git a/Engine/Plugins/VirtualProduction/Takes/Source/TakeRecorderSources/Private/TakeRecorderSourcesUtils.h b/Engine/Plugins/VirtualProduction/Takes/Source/TakeRecorderSources/Private/TakeRecorderSourcesUtils.h
--- a/Engine/Plugins/VirtualProduction/Takes/Source/TakeRecorderSources/Private/TakeRecorderSourcesUtils.h
+++ b/Engine/Plugins/VirtualProduction/Takes/Source/TakeRecorderSources/Private/TakeRecorderSourcesUtils.h
@@ -93,6 +93,11 @@ namespace TakeRecorderSourcesUtils
 		return false;
 	}
 
+	static bool IsActorRecordable(const AActor* Actor)
+	{
+		return !Actor->GetClass()->HasAnyClassFlags(CLASS_NotPlaceable);
+	};
+
 	/**
 	* Get the object binding for a given actor that is being recorded. An actor can either be a Possessable or a Spawnable but we only have pointers
 	* to the original object being recorded. To solve this, we iterate through each actor being recorded and ask it what Guid it ended up with which
