Take Recorder: Default to not recording timecode

diff --git a/Engine/Plugins/VirtualProduction/Takes/Source/TakeRecorder/Private/Recorder/TakeRecorderParameters.cpp b/Engine/Plugins/VirtualProduction/Takes/Source/TakeRecorder/Private/Recorder/TakeRecorderParameters.cpp
--- a/Engine/Plugins/VirtualProduction/Takes/Source/TakeRecorder/Private/Recorder/TakeRecorderParameters.cpp
+++ b/Engine/Plugins/VirtualProduction/Takes/Source/TakeRecorder/Private/Recorder/TakeRecorderParameters.cpp
@@ -18,7 +18,7 @@ FTakeRecorderUserParameters::FTakeRecorderUserParameters()
 
 FTakeRecorderProjectParameters::FTakeRecorderProjectParameters()
 	: bStartAtCurrentTimecode(true)
-	, bRecordTimecode(true)
+	, bRecordTimecode(false)
 	, bRecordSourcesIntoSubSequences(false)
 	, bRecordToPossessable(false)
 {}
