Take Recorder: Make sure attach section is added to the attach track

diff --git a/Engine/Plugins/VirtualProduction/Takes/Source/TakeTrackRecorders/Private/MovieScene3DAttachTrackRecorder.cpp b/Engine/Plugins/VirtualProduction/Takes/Source/TakeTrackRecorders/Private/MovieScene3DAttachTrackRecorder.cpp
--- a/Engine/Plugins/VirtualProduction/Takes/Source/TakeTrackRecorders/Private/MovieScene3DAttachTrackRecorder.cpp
+++ b/Engine/Plugins/VirtualProduction/Takes/Source/TakeTrackRecorders/Private/MovieScene3DAttachTrackRecorder.cpp
@@ -50,6 +50,8 @@ void UMovieScene3DAttachTrackRecorder::RecordSampleImpl(const FQualifiedFrameTim
 			if (!MovieSceneSection.IsValid() || AttachedToActor != ActorAttachedTo.Get())
 			{
 				MovieSceneSection = Cast<UMovieScene3DAttachSection>(AttachTrack->CreateNewSection());
+				AttachTrack->AddSection(*MovieSceneSection);
+
 				MovieSceneSection->AttachSocketName = SocketName;
 				MovieSceneSection->AttachComponentName = ComponentName;
 
diff --git a/Engine/Source/Editor/SequenceRecorderSections/Private/MovieScene3DAttachSectionRecorder.cpp b/Engine/Source/Editor/SequenceRecorderSections/Private/MovieScene3DAttachSectionRecorder.cpp
--- a/Engine/Source/Editor/SequenceRecorderSections/Private/MovieScene3DAttachSectionRecorder.cpp
+++ b/Engine/Source/Editor/SequenceRecorderSections/Private/MovieScene3DAttachSectionRecorder.cpp
@@ -60,6 +60,7 @@ void FMovieScene3DAttachSectionRecorder::Record(float CurrentTime)
 			if(!MovieSceneSection.IsValid() || AttachedToActor != ActorAttachedTo.Get())
 			{
 				MovieSceneSection = Cast<UMovieScene3DAttachSection>(AttachTrack->CreateNewSection());
+				AttachTrack->AddSection(*MovieSceneSection);
 
 				FFrameRate TickResolution = MovieSceneSection->GetTypedOuter<UMovieScene>()->GetTickResolution();
 				FFrameNumber CurrentFrame = (CurrentTime * TickResolution).FloorToFrame();
