Take Recorder: Add missing icons and styles

diff --git a/Engine/Plugins/VirtualProduction/Takes/Source/TakeRecorder/Private/TakeRecorderStyle.cpp b/Engine/Plugins/VirtualProduction/Takes/Source/TakeRecorder/Private/TakeRecorderStyle.cpp
--- a/Engine/Plugins/VirtualProduction/Takes/Source/TakeRecorder/Private/TakeRecorderStyle.cpp
+++ b/Engine/Plugins/VirtualProduction/Takes/Source/TakeRecorder/Private/TakeRecorderStyle.cpp
@@ -113,7 +113,20 @@ FTakeRecorderStyle::FTakeRecorderStyle()
 		.SetUndeterminedHoveredImage( FSlateBoxBrush(RootToCoreContentDir(TEXT("Common/RoundedSelection_16x.png")),  4.0f/16.0f, SelectionColor ) )
 		.SetUndeterminedPressedImage( FSlateBoxBrush(RootToCoreContentDir(TEXT("Common/RoundedSelection_16x.png")),  4.0f/16.0f, SelectionColor_Inactive) );
 
+	const FCheckBoxStyle ToggleButtonIndicatorStyle = FCheckBoxStyle()
+		.SetCheckBoxType(ESlateCheckBoxType::ToggleButton)
+		.SetUncheckedImage( FSlateNoResource() )
+		.SetUncheckedHoveredImage( FSlateNoResource() )
+		.SetUncheckedPressedImage( FSlateNoResource() )
+		.SetCheckedImage( FSlateBoxBrush(RootToCoreContentDir(TEXT("Common/RoundedSelection_16x.png")),  4.0f/16.0f, SelectionColor_Pressed ) )
+		.SetCheckedHoveredImage( FSlateBoxBrush(RootToCoreContentDir(TEXT("Common/RoundedSelection_16x.png")),  4.0f/16.0f, SelectionColor_Pressed ) )
+		.SetCheckedPressedImage( FSlateBoxBrush(RootToCoreContentDir(TEXT("Common/RoundedSelection_16x.png")),  4.0f/16.0f, SelectionColor_Pressed ) )
+		.SetUndeterminedImage( FSlateBoxBrush(RootToCoreContentDir(TEXT("Common/RoundedSelection_16x.png")),  4.0f/16.0f, SelectionColor_Inactive) )
+		.SetUndeterminedHoveredImage( FSlateBoxBrush(RootToCoreContentDir(TEXT("Common/RoundedSelection_16x.png")),  4.0f/16.0f, SelectionColor_Inactive) )
+		.SetUndeterminedPressedImage( FSlateBoxBrush(RootToCoreContentDir(TEXT("Common/RoundedSelection_16x.png")),  4.0f/16.0f, SelectionColor_Inactive) );
+
 	Set( "ToggleButtonCheckbox", ToggleButtonStyle );	
+	Set( "ToggleButtonIndicatorCheckbox", ToggleButtonIndicatorStyle );	
 
 	Set("WhiteBrush", new FSlateColorBrush(FLinearColor::White));
 	Set("Button", Button);
@@ -181,6 +194,7 @@ FTakeRecorderStyle::FTakeRecorderStyle()
 	Set("TakeRecorder.SequencerButton", new FSlateImageBrush(RootToContentDir(TEXT("Sequencer.png")), Icon20x20));
 	Set("TakeRecorder.ReviewRecordingButton", new FSlateImageBrush(RootToContentDir(TEXT("ReviewRecording.png")), Icon20x20));
 	Set("TakeRecorder.MarkFrame", new FSlateImageBrush(RootToContentDir(TEXT("MarkFrame.png")), Icon20x20));
+	Set("TakeRecorder.SequenceToRecordIntoButton", new FSlateImageBrush(RootToContentDir(TEXT("SequenceToRecordInto.png")), Icon20x20));
 
 	Set("TakeRecorder.Slate", new FSlateColorBrush(MediumGrey));
 	Set("TakeRecorder.Slate.ClapperBackground", new FSlateColorBrush(AlmostWhite));
