Take Recorder: Change recording time to be driven by Sequencer time which is in relative timecode mode. This addresses jittering and inaccuracies in plotting recorded values.

- Removed GetCurrentRecordingFrame and replaced it with the time from Sequencer which is driven by the timecode clock.
- Removed unused code StartRecordingRecursive
- Removed unused TimecodeSource

diff --git a/Engine/Plugins/VirtualProduction/Takes/Source/TakeRecorder/Private/Recorder/TakeRecorder.cpp b/Engine/Plugins/VirtualProduction/Takes/Source/TakeRecorder/Private/Recorder/TakeRecorder.cpp
--- a/Engine/Plugins/VirtualProduction/Takes/Source/TakeRecorder/Private/Recorder/TakeRecorder.cpp
+++ b/Engine/Plugins/VirtualProduction/Takes/Source/TakeRecorder/Private/Recorder/TakeRecorder.cpp
@@ -759,8 +759,6 @@ UWorld* UTakeRecorder::GetWorld() const
 
 void UTakeRecorder::Tick(float DeltaTime)
 {
-	FTimecode TimecodeSource = FApp::GetTimecode();
-
 	if (State == ETakeRecorderState::CountingDown)
 	{
 		NumberOfTicksAfterPre = 0;
@@ -781,28 +779,30 @@ void UTakeRecorder::Tick(float DeltaTime)
 	else if (State == ETakeRecorderState::TickingAfterPre)
 	{
 		NumberOfTicksAfterPre = 0;
-		Start(TimecodeSource);
-		InternalTick(TimecodeSource, 0.0f);
+		Start();
+		InternalTick(0.0f);
 	}
 	else if (State == ETakeRecorderState::Started)
 	{
-		InternalTick(TimecodeSource, DeltaTime);
+		InternalTick(DeltaTime);
 	}
 }
 
-void UTakeRecorder::InternalTick(const FTimecode& InTimecodeSource, float DeltaTime)
+void UTakeRecorder::InternalTick(float DeltaTime)
 {
+	TSharedPtr<ISequencer> Sequencer = WeakSequencer.Pin();
+	FQualifiedFrameTime SequencerTime = Sequencer.IsValid() ? Sequencer->GetGlobalTime() : FQualifiedFrameTime();
+	
 	UTakeRecorderSources* Sources = SequenceAsset->FindOrAddMetaData<UTakeRecorderSources>();
 	if (!Parameters.bDisableRecordingAndSave)
 	{
-		CurrentFrameTime = Sources->TickRecording(SequenceAsset, InTimecodeSource, DeltaTime);
+		CurrentFrameTime = Sources->TickRecording(SequenceAsset, SequencerTime, DeltaTime);
 	}
 	else
 	{
-		CurrentFrameTime = Sources->AdvanceTime(DeltaTime);
+		CurrentFrameTime = Sources->AdvanceTime(SequencerTime, DeltaTime);
 	}
 
-	TSharedPtr<ISequencer> Sequencer = WeakSequencer.Pin();
 	if (Sequencer.IsValid())
 	{
 		FAnimatedRange Range = Sequencer->GetViewRange();
@@ -869,9 +869,21 @@ void UTakeRecorder::PreRecord()
 
 	Sources->SetSettings(TakeRecorderSourcesSettings);
 
+	TSharedPtr<ISequencer> Sequencer = WeakSequencer.Pin();
 	if (!Parameters.bDisableRecordingAndSave)
 	{
-		Sources->PreRecording(SequenceAsset, Parameters.User.bAutoSerialize ? &ManifestSerializer : nullptr);
+		UMovieScene* MovieScene = SequenceAsset->GetMovieScene();
+		FQualifiedFrameTime SequencerTime;
+		if (MovieScene)
+		{
+			FTimecode Timecode = FApp::GetTimecode();
+			FFrameRate FrameRate = MovieScene->GetDisplayRate();
+			FFrameRate TickResolution = MovieScene->GetTickResolution();
+			FFrameNumber PlaybackStartFrame = Parameters.Project.bStartAtCurrentTimecode ? FFrameRate::TransformTime(FFrameTime(Timecode.ToFrameNumber(FrameRate)), FrameRate, TickResolution).FloorToFrame() : MovieScene->GetPlaybackRange().GetLowerBoundValue();
+			SequencerTime = FQualifiedFrameTime(PlaybackStartFrame, TickResolution);
+		}
+
+		Sources->PreRecording(SequenceAsset, SequencerTime, Parameters.User.bAutoSerialize ? &ManifestSerializer : nullptr);
 	}
 	else
 	{
@@ -879,15 +891,16 @@ void UTakeRecorder::PreRecord()
 	}
 
 	// Refresh sequencer in case the movie scene data has mutated (ie. existing object bindings removed because they will be recorded again)
-	TSharedPtr<ISequencer> Sequencer = WeakSequencer.Pin();
 	if (Sequencer.IsValid())
 	{
 		Sequencer->RefreshTree();
 	}
 }
 
-void UTakeRecorder::Start(const FTimecode& InTimecodeSource)
+void UTakeRecorder::Start()
 {
+	FTimecode Timecode = FApp::GetTimecode();
+
 	State = ETakeRecorderState::Started;
 
 	TSharedPtr<ISequencer> Sequencer = WeakSequencer.Pin();
@@ -913,7 +926,7 @@ void UTakeRecorder::Start(const FTimecode& InTimecodeSource)
 
 		FFrameRate FrameRate = MovieScene->GetDisplayRate();
 		FFrameRate TickResolution = MovieScene->GetTickResolution();
-		FFrameNumber PlaybackStartFrame = Parameters.Project.bStartAtCurrentTimecode ? FFrameRate::TransformTime(FFrameTime(InTimecodeSource.ToFrameNumber(FrameRate)), FrameRate, TickResolution).FloorToFrame() : MovieScene->GetPlaybackRange().GetLowerBoundValue();
+		FFrameNumber PlaybackStartFrame = Parameters.Project.bStartAtCurrentTimecode ? FFrameRate::TransformTime(FFrameTime(Timecode.ToFrameNumber(FrameRate)), FrameRate, TickResolution).FloorToFrame() : MovieScene->GetPlaybackRange().GetLowerBoundValue();
 
 		// Transform all the sections to start the playback start frame
 		FFrameNumber DeltaFrame = PlaybackStartFrame - MovieScene->GetPlaybackRange().GetLowerBoundValue();
@@ -944,11 +957,16 @@ void UTakeRecorder::Start(const FTimecode& InTimecodeSource)
 	UTakeMetaData* AssetMetaData = SequenceAsset->FindMetaData<UTakeMetaData>();
 	FDateTime UtcNow = FDateTime::UtcNow();
 	AssetMetaData->SetTimestamp(UtcNow);
-	AssetMetaData->SetTimecodeIn(InTimecodeSource);
+	AssetMetaData->SetTimecodeIn(Timecode);
 
 	if (!Parameters.bDisableRecordingAndSave)
 	{
-		Sources->StartRecording(SequenceAsset, InTimecodeSource, Parameters.User.bAutoSerialize ? &ManifestSerializer : nullptr);
+		FQualifiedFrameTime SequencerTime = Sequencer.IsValid() ? Sequencer->GetGlobalTime() : FQualifiedFrameTime();
+
+		Sources->StartRecording(SequenceAsset, SequencerTime, Parameters.User.bAutoSerialize ? &ManifestSerializer : nullptr);
+
+		// Record immediately so that there's a key on the first frame of recording
+		Sources->TickRecording(SequenceAsset, SequencerTime, 0.1f);
 	}
 
 	if (!ShouldShowNotifications())
@@ -975,8 +993,6 @@ void UTakeRecorder::Stop()
 
 	TGuardValue<bool> ReentrantGuard(bStoppedRecording, true);
 
-	FTimecode TimecodeOut = FApp::GetTimecode();
-	
 	if (Parameters.TakeRecorderMode == ETakeRecorderMode::RecordNewSequence)
 	{
 		USequencerSettings* SequencerSettings = USequencerSettingsContainer::GetOrCreate<USequencerSettings>(TEXT("TakeRecorderSequenceEditor"));
@@ -1026,10 +1042,8 @@ void UTakeRecorder::Stop()
 			TRange<FFrameNumber> Range = MovieScene->GetPlaybackRange();
 
 			//Set Range to what we recorded instead of that large number, this let's us reliably set camera cut times.
-			FFrameRate FrameRate = MovieScene->GetDisplayRate();
-			FFrameRate TickResolution = MovieScene->GetTickResolution();
-			FFrameNumber PlaybackEndFrame = Parameters.Project.bStartAtCurrentTimecode ? FFrameRate::TransformTime(FFrameTime(TimecodeOut.ToFrameNumber(FrameRate)), FrameRate, TickResolution).CeilToFrame() : CurrentFrameTime.FrameNumber;
-			
+			FFrameNumber PlaybackEndFrame = CurrentFrameTime.FrameNumber;
+
 			if (StopRecordingFrame.IsSet())
 			{
 				PlaybackEndFrame = StopRecordingFrame.GetValue();
@@ -1070,7 +1084,14 @@ void UTakeRecorder::Stop()
 		UTakeMetaData* AssetMetaData = SequenceAsset->FindMetaData<UTakeMetaData>();
 		check(AssetMetaData);
 
-		AssetMetaData->SetTimecodeOut(TimecodeOut);
+		if (MovieScene)
+		{
+			FFrameRate DisplayRate = MovieScene->GetDisplayRate();
+			FFrameRate TickResolution = MovieScene->GetTickResolution();
+			FTimecode Timecode = FTimecode::FromFrameNumber(FFrameRate::TransformTime(CurrentFrameTime, TickResolution, DisplayRate).FloorToFrame(), DisplayRate);
+
+			AssetMetaData->SetTimecodeOut(Timecode);
+		}
 
 		if (GEditor && GEditor->GetEditorWorldContext().World())
 		{
diff --git a/Engine/Plugins/VirtualProduction/Takes/Source/TakeRecorder/Public/Recorder/TakeRecorder.h b/Engine/Plugins/VirtualProduction/Takes/Source/TakeRecorder/Public/Recorder/TakeRecorder.h
--- a/Engine/Plugins/VirtualProduction/Takes/Source/TakeRecorder/Public/Recorder/TakeRecorder.h
+++ b/Engine/Plugins/VirtualProduction/Takes/Source/TakeRecorder/Public/Recorder/TakeRecorder.h
@@ -159,7 +159,7 @@ private:
 	/**
 	 * Called after PreRecord To Start
 	 */
-	void Start(const FTimecode& InTimecodeSource);
+	void Start();
 
 	/**
 	 * Ticked by a tickable game object to performe any necessary time-sliced logic
@@ -205,7 +205,7 @@ private:
 private:
 
 	/** Called by Tick and Start to make sure we record at start */
-	void InternalTick(const FTimecode& InTimecodeSource, float DeltaTime);
+	void InternalTick(float DeltaTime);
 
 	virtual UWorld* GetWorld() const override;
 
diff --git a/Engine/Plugins/VirtualProduction/Takes/Source/TakeRecorderSources/Private/TakeRecorderActorSource.cpp b/Engine/Plugins/VirtualProduction/Takes/Source/TakeRecorderSources/Private/TakeRecorderActorSource.cpp
--- a/Engine/Plugins/VirtualProduction/Takes/Source/TakeRecorderSources/Private/TakeRecorderActorSource.cpp
+++ b/Engine/Plugins/VirtualProduction/Takes/Source/TakeRecorderSources/Private/TakeRecorderActorSource.cpp
@@ -294,7 +294,8 @@ TArray<UTakeRecorderSource*> UTakeRecorderActorSource::PreRecording(ULevelSequen
 	}
 	// Now we need to create the section recorders for each of the enabled properties based on the property map.
 	// Any components spawned at runtime will get picked up on Tick and have section recorders created for them mid-recording.
-	CreateSectionRecordersRecursive(ActorToRecord, RecordedProperties);
+	TArray<UObject*> TraversedObjects;
+	CreateSectionRecordersRecursive(ActorToRecord, RecordedProperties, TraversedObjects);
 
 	// Update our cached list of components so that we don't detect them all as new components on the first Tick
 	GetAllComponents(CachedComponentList, false);
@@ -356,8 +357,15 @@ void UTakeRecorderActorSource::StartRecording(const FTimecode& InSectionStartTim
 	}
 }
 
-void UTakeRecorderActorSource::CreateSectionRecordersRecursive(UObject* ObjectToRecord, UActorRecorderPropertyMap* PropertyMap)
+void UTakeRecorderActorSource::CreateSectionRecordersRecursive(UObject* ObjectToRecord, UActorRecorderPropertyMap* PropertyMap, TArray<UObject*>& TraversedObjects)
 {
+	if (TraversedObjects.Contains(ObjectToRecord))
+	{
+		return;
+	}
+
+	TraversedObjects.Add(ObjectToRecord);
+
 	FGuid Guid = CachedObjectBindingGuid;
 	if (ObjectToRecord->IsA<UActorComponent>())
 	{
@@ -513,7 +521,7 @@ void UTakeRecorderActorSource::CreateSectionRecordersRecursive(UObject* ObjectTo
 		UObject* ChildObject = Child->RecordedObject.Get();
 		if (ChildObject)
 		{
-			CreateSectionRecordersRecursive(ChildObject, Child);
+			CreateSectionRecordersRecursive(ChildObject, Child, TraversedObjects);
 		}
 		else
 		{
@@ -552,6 +560,7 @@ void UTakeRecorderActorSource::TickRecording(const FQualifiedFrameTime& CurrentS
 		}
 	}
 
+	TArray<UObject*> TraversedObjects;
 	for (UActorComponent* AddedComponent : NewComponentsAdded)
 	{
 		if (Target.IsValid())
@@ -575,7 +584,7 @@ void UTakeRecorderActorSource::TickRecording(const FQualifiedFrameTime& CurrentS
 			RebuildRecordedPropertyMapRecursive(AddedComponent, ComponentPropertyMap);
 
 			// Create the Section Recorders required
-			CreateSectionRecordersRecursive(AddedComponent, ComponentPropertyMap);
+			CreateSectionRecordersRecursive(AddedComponent, ComponentPropertyMap, TraversedObjects);
 
 			// Update our numbers on the display
 			UpdateCachedNumberOfRecordedProperties();
@@ -1561,14 +1570,6 @@ void UTakeRecorderActorSource::CreateNewActorSourceForReferencedActors()
 		ActorSource->PostEditChangeProperty(PropertyChangedEvent);
 	}
 	NewReferencedActors.Reset();
-
-	// If this actor reference was added at runtime we need to immediately start recording it. If this was added during PreRecording
-	// we don't want to call StartRecordingSource as the UTakeRecorderSources will do that automatically.
-	const bool bStartRecordingImmediately = false;
-	if (bStartRecordingImmediately)
-	{
-		SourcesList->StartRecordingSource(NewSources,TimecodeSource);
-	}
 }
 
 bool UTakeRecorderActorSource::IsOtherActorBeingRecorded(AActor* OtherActor) const
diff --git a/Engine/Plugins/VirtualProduction/Takes/Source/TakeRecorderSources/Private/TakeRecorderActorSource.h b/Engine/Plugins/VirtualProduction/Takes/Source/TakeRecorderSources/Private/TakeRecorderActorSource.h
--- a/Engine/Plugins/VirtualProduction/Takes/Source/TakeRecorderSources/Private/TakeRecorderActorSource.h
+++ b/Engine/Plugins/VirtualProduction/Takes/Source/TakeRecorderSources/Private/TakeRecorderActorSource.h
@@ -175,7 +175,7 @@ protected:
 	* This is called when recording starts to generate the Section Recorders for the actor and all components that it currently has,
 	* as well as again during runtime for any newly added components.
 	*/
-	void CreateSectionRecordersRecursive(UObject* ObjectToRecord, UActorRecorderPropertyMap* PropertyMap);
+	void CreateSectionRecordersRecursive(UObject* ObjectToRecord, UActorRecorderPropertyMap* PropertyMap, TArray<UObject*>& TraversedObjects);
 
 	/** Update our cached properties for what will be recorded. Done here so the UI doesn't have to iterate through map every frame. */
 	void UpdateCachedNumberOfRecordedProperties();
diff --git a/Engine/Plugins/VirtualProduction/Takes/Source/TakeRecorderSources/Private/TakeRecorderCameraCutSource.cpp b/Engine/Plugins/VirtualProduction/Takes/Source/TakeRecorderSources/Private/TakeRecorderCameraCutSource.cpp
--- a/Engine/Plugins/VirtualProduction/Takes/Source/TakeRecorderSources/Private/TakeRecorderCameraCutSource.cpp
+++ b/Engine/Plugins/VirtualProduction/Takes/Source/TakeRecorderSources/Private/TakeRecorderCameraCutSource.cpp
@@ -82,7 +82,7 @@ void UTakeRecorderCameraCutSource::TickRecording(const FQualifiedFrameTime& Curr
 
 		// This has to be called after setting the Target and propagating the change event so that it has a chance to know what to record
 		// about the actor.
-		OwningSources->StartRecordingSource(TArray<UTakeRecorderSource*>({ ActorSource }), FApp::GetTimecode());
+		OwningSources->StartRecordingSource(TArray<UTakeRecorderSource*>({ ActorSource }), CurrentTime);
 
 		NewActorSources.Add(ActorSource);
 	}
diff --git a/Engine/Plugins/VirtualProduction/Takes/Source/TakeRecorderSources/Private/TakeRecorderMicrophoneAudioSource.cpp b/Engine/Plugins/VirtualProduction/Takes/Source/TakeRecorderSources/Private/TakeRecorderMicrophoneAudioSource.cpp
--- a/Engine/Plugins/VirtualProduction/Takes/Source/TakeRecorderSources/Private/TakeRecorderMicrophoneAudioSource.cpp
+++ b/Engine/Plugins/VirtualProduction/Takes/Source/TakeRecorderSources/Private/TakeRecorderMicrophoneAudioSource.cpp
@@ -203,11 +203,11 @@ void UTakeRecorderMicrophoneAudioSource::StopRecording(class ULevelSequence* InS
 
 		UMovieSceneAudioSection* NewAudioSection = NewObject<UMovieSceneAudioSection>(CachedAudioTrack.Get(), UMovieSceneAudioSection::StaticClass());
 
-		FFrameNumber RecordStartFrame = Sources && Sources->GetSettings().bStartAtCurrentTimecode ? FFrameRate::TransformTime(FFrameTime(TimecodeSource.ToFrameNumber(DisplayRate)), DisplayRate, TickResolution).FloorToFrame() : MovieScene->GetPlaybackRange().GetLowerBoundValue();
+		FFrameNumber RecordStartFrame = MovieScene->GetPlaybackRange().GetLowerBoundValue();
 
 		NewAudioSection->SetSound(RecordedAudio);
 		NewAudioSection->SetRange(TRange<FFrameNumber>(RecordStartFrame, RecordStartFrame + (RecordedAudio->GetDuration() * TickResolution).CeilToFrame()));
-		NewAudioSection->TimecodeSource = TimecodeSource;
+		NewAudioSection->TimecodeSource = FTimecode::FromFrameNumber(RecordStartFrame, DisplayRate);
 
 		CachedAudioTrack->AddSection(*NewAudioSection);
 
diff --git a/Engine/Plugins/VirtualProduction/Takes/Source/TakeRecorderSources/Private/TakeRecorderNearbySpawnedActorSource.cpp b/Engine/Plugins/VirtualProduction/Takes/Source/TakeRecorderSources/Private/TakeRecorderNearbySpawnedActorSource.cpp
--- a/Engine/Plugins/VirtualProduction/Takes/Source/TakeRecorderSources/Private/TakeRecorderNearbySpawnedActorSource.cpp
+++ b/Engine/Plugins/VirtualProduction/Takes/Source/TakeRecorderSources/Private/TakeRecorderNearbySpawnedActorSource.cpp
@@ -26,7 +26,7 @@ UTakeRecorderNearbySpawnedActorSource::UTakeRecorderNearbySpawnedActorSource(con
 	TrackTint = FColor(167, 167, 59);
 }
 
-void UTakeRecorderNearbySpawnedActorSource::StartRecording(const FTimecode& InSectionStartTimecode, const FFrameNumber& InSectionFirstFrame, class ULevelSequence* InSequence)
+TArray<UTakeRecorderSource*> UTakeRecorderNearbySpawnedActorSource::PreRecording(ULevelSequence* InSequence, FMovieSceneSequenceID InSequenceID, ULevelSequence* InMasterSequence, FManifestSerializer* InManifestSerializer)
 {
 	// Get the sources, add callbacks for new spawned
 	UWorld* World = TakeRecorderSourcesUtils::GetSourceWorld(InSequence);
@@ -40,6 +40,8 @@ void UTakeRecorderNearbySpawnedActorSource::StartRecording(const FTimecode& InSe
 			ActorSpawningDelegateHandles.Add(World, NewHandle);
 		}
 	}
+
+	return TArray<UTakeRecorderSource*>();
 }
 
 TArray<UTakeRecorderSource*> UTakeRecorderNearbySpawnedActorSource::PostRecording(class ULevelSequence* InSequence, class ULevelSequence* InMasterSequence)
@@ -188,7 +190,7 @@ void UTakeRecorderNearbySpawnedActorSource::HandleActorSpawned(AActor* Actor, cl
 
 	// This has to be called after setting the Target and propegating the change event so that it has a chance to know what to record
 	// about the actor.
-	Sources->StartRecordingSource(TArray<UTakeRecorderSource*>({ ActorSource }), FApp::GetTimecode());
+	Sources->StartRecordingSource(TArray<UTakeRecorderSource*>({ ActorSource }), Sources->GetCachedFrameTime());
 
 	SpawnedActorSources.Add(ActorSource);
 }
diff --git a/Engine/Plugins/VirtualProduction/Takes/Source/TakeRecorderSources/Private/TakeRecorderNearbySpawnedActorSource.h b/Engine/Plugins/VirtualProduction/Takes/Source/TakeRecorderSources/Private/TakeRecorderNearbySpawnedActorSource.h
--- a/Engine/Plugins/VirtualProduction/Takes/Source/TakeRecorderSources/Private/TakeRecorderNearbySpawnedActorSource.h
+++ b/Engine/Plugins/VirtualProduction/Takes/Source/TakeRecorderSources/Private/TakeRecorderNearbySpawnedActorSource.h
@@ -35,7 +35,7 @@ public:
 private:
 
 	// UTakeRecorderSource
-	virtual void StartRecording(const FTimecode& InSectionStartTimecode, const FFrameNumber& InSectionFirstFrame, class ULevelSequence* InSequence) override;
+	virtual TArray<UTakeRecorderSource*> PreRecording(ULevelSequence* InSequence, FMovieSceneSequenceID InSequenceID, ULevelSequence* InMasterSequence, FManifestSerializer* InManifestSerializer) override;
 	virtual TArray<UTakeRecorderSource*> PostRecording(class ULevelSequence* InSequence, class ULevelSequence* InMasterSequence) override;
 	virtual FText GetDisplayTextImpl() const override;
 	virtual FText GetDescriptionTextImpl() const override;
diff --git a/Engine/Plugins/VirtualProduction/Takes/Source/TakeTrackRecorders/Private/MovieSceneAnimationTrackRecorder.cpp b/Engine/Plugins/VirtualProduction/Takes/Source/TakeTrackRecorders/Private/MovieSceneAnimationTrackRecorder.cpp
--- a/Engine/Plugins/VirtualProduction/Takes/Source/TakeTrackRecorders/Private/MovieSceneAnimationTrackRecorder.cpp
+++ b/Engine/Plugins/VirtualProduction/Takes/Source/TakeTrackRecorders/Private/MovieSceneAnimationTrackRecorder.cpp
@@ -195,13 +195,10 @@ void UMovieSceneAnimationTrackRecorder::StopRecordingImpl()
 {
 	AnimationSerializer.Close();
 
-	if (SkeletalMeshComponent.IsValid())
-	{
-		// Legacy Animation Recorder allowed recording into an animation asset directly and not creating an movie section
-		const bool bShowAnimationAssetCreatedToast = false;
-		InitialRootTransform = AnimationRecorder.Recorder.Get()->GetInitialRootTransform();
-		AnimationRecorder.FinishRecording(bShowAnimationAssetCreatedToast);
-	}
+	// Legacy Animation Recorder allowed recording into an animation asset directly and not creating an movie section
+	const bool bShowAnimationAssetCreatedToast = false;
+	InitialRootTransform = AnimationRecorder.Recorder.Get()->GetInitialRootTransform();
+	AnimationRecorder.FinishRecording(bShowAnimationAssetCreatedToast);
 }
 
 void UMovieSceneAnimationTrackRecorder::FinalizeTrackImpl()
diff --git a/Engine/Plugins/VirtualProduction/Takes/Source/TakesCore/Private/TakeRecorderSources.cpp b/Engine/Plugins/VirtualProduction/Takes/Source/TakesCore/Private/TakeRecorderSources.cpp
--- a/Engine/Plugins/VirtualProduction/Takes/Source/TakesCore/Private/TakeRecorderSources.cpp
+++ b/Engine/Plugins/VirtualProduction/Takes/Source/TakesCore/Private/TakeRecorderSources.cpp
@@ -81,169 +81,6 @@ void UTakeRecorderSources::PostEditChangeProperty(FPropertyChangedEvent& Propert
 	}
 }
 
-void UTakeRecorderSources::StartRecordingRecursive(TArray<UTakeRecorderSource*> InSources, ULevelSequence* InMasterSequence, const FTimecode& Timecode, FManifestSerializer* InManifestSerializer)
-{
-	TArray<UTakeRecorderSource*> NewSources;
-
-	// Optionally create a folder in the Sequencer UI that will contain this source. We don't want sub-sequences to have folders
-	// created for their sources as you would end up with a Subscene with one item in it hidden inside of a folder, so instead
-	// only the master sequence gets folders created.
-	const bool bCreateSequencerFolders = true;
-	for (UTakeRecorderSource* Source : InSources)
-	{
-		if (Source->bEnabled)
-		{
-			ULevelSequence* TargetSequence = InMasterSequence;
-			FMovieSceneSequenceID TargetSequenceID = MovieSceneSequenceID::Root;
-
-			// The Sequencer Take system is built around swapping out sub-sequences. If they want to use this system, we create a sub-sequence
-			// for the Source and tell it to write into this sub-sequence instead of the master sequence. We then keep track of which Source
-			// is using which sub-sequence so that we can push the correct sequence for all points of the Source's recording lifecycle.
-			if (Settings.bRecordSourcesIntoSubSequences && Source->SupportsSubscenes())
-			{
-				const FString& SubSequenceTrackName = ObjectTools::SanitizeObjectName(Source->GetSubsceneTrackName(InMasterSequence));
-				const FString& SubSequenceAssetName = ObjectTools::SanitizeObjectName(Source->GetSubsceneAssetName(InMasterSequence));
-
-				TargetSequence = CreateSubSequenceForSource(InMasterSequence, SubSequenceTrackName, SubSequenceAssetName);
-				TargetSequence->GetMovieScene()->TimecodeSource = Timecode;
-
-				// If there's already a Subscene Track for our sub-sequence we need to remove that track before create a new one. No data is lost in this process as the
-				// sequence that the subscene points to has been copied by CreateSubSequenceForSource so a new track pointed to the new subsequence includes all the old data.
-				const FString SequenceName = FPaths::GetBaseFilename(TargetSequence->GetPathName());
-				UMovieSceneSubTrack* SubsceneTrack = nullptr;
-
-				for (UMovieSceneTrack* Track : InMasterSequence->GetMovieScene()->GetMasterTracks())
-				{
-					if (Track->IsA<UMovieSceneSubTrack>())
-					{
-						if (Track->GetDisplayName().ToString() == SubSequenceTrackName)
-						{
-							SubsceneTrack = CastChecked<UMovieSceneSubTrack>(Track);
-							SubsceneTrack->RemoveAllAnimationData();
-						}
-					}
-				}
-
-				// We need to add the new subsequence to the master sequence immediately so that it shows up in the UI and you can tell that things
-				// are being recorded, otherwise they don't show up until recording stops and then it magically pops in.
-				if (!SubsceneTrack)
-				{
-					SubsceneTrack = CastChecked<UMovieSceneSubTrack>(InMasterSequence->GetMovieScene()->AddMasterTrack(UMovieSceneSubTrack::StaticClass()));
-				}
-
-				// Track should not be transactional during the recording process
-				SubsceneTrack->ClearFlags(RF_Transactional);
-
-				// We create a new sub track for every Source so that we can name the Subtrack after the Source instead of just the sections within it.
-				SubsceneTrack->SetDisplayName(FText::FromString(Source->GetSubsceneTrackName(InMasterSequence)));
-				SubsceneTrack->SetColorTint(Source->TrackTint);
-
-				// When we create the Subscene Track we'll make sure a folder is created for it to sort into and add the new Subscene Track as a child of it.
-				if (bCreateSequencerFolders)
-				{
-					UMovieSceneFolder* Folder = AddFolderForSource(Source, InMasterSequence->GetMovieScene());
-					Folder->AddChildMasterTrack(SubsceneTrack);
-				}
-
-				// We initialize the sequence to start at the current time.
-				// We'll have to update these sections each frame as the recording progresses so they appear to get longer like normal
-				// tracks do as we record into them.
-				FFrameNumber RecordStartFrame = Settings.bStartAtCurrentTimecode ? FFrameRate::TransformTime(FFrameTime(Timecode.ToFrameNumber(TargetLevelSequenceDisplayRate)), TargetLevelSequenceDisplayRate, TargetLevelSequenceTickResolution).FloorToFrame() : InMasterSequence->GetMovieScene()->GetPlaybackRange().GetLowerBoundValue();
-				UMovieSceneSubSection* NewSubSection = SubsceneTrack->AddSequence(TargetSequence, RecordStartFrame, 0);
-				SetSectionStartTimecode(NewSubSection, Timecode, TargetLevelSequenceDisplayRate, TargetLevelSequenceTickResolution);
-
-				// Section should not be transactional during the recording process
-				NewSubSection->ClearFlags(RF_Transactional);
-
-				NewSubSection->SetRowIndex(SubsceneTrack->GetMaxRowIndex() + 1);
-				SubsceneTrack->FixRowIndices();
-
-				TargetSequenceID = NewSubSection->GetSequenceID();
-
-				ActiveSubSections.Add(NewSubSection);
-				if (InManifestSerializer)
-				{
-					FName SerializedType("SubSequence");
-					FManifestProperty  ManifestProperty(SubSequenceAssetName, SerializedType, FGuid());
-					InManifestSerializer->WriteFrameData(InManifestSerializer->FramesWritten, ManifestProperty);
-
-					FString AssetPath = InManifestSerializer->GetLocalCaptureDir();
-
-					IPlatformFile& PlatformFile = FPlatformFileManager::Get().GetPlatformFile();
-					if (!PlatformFile.DirectoryExists(*AssetPath))
-					{
-						PlatformFile.CreateDirectory(*AssetPath);
-					}
-
-					AssetPath = AssetPath / SubSequenceAssetName;
-					if (!PlatformFile.DirectoryExists(*AssetPath))
-					{
-						PlatformFile.CreateDirectory(*AssetPath);
-					}
-
-					TSharedPtr<FManifestSerializer> NewManifestSerializer = MakeShared<FManifestSerializer>();
-					CreatedManifestSerializers.Add(NewManifestSerializer);
-					InManifestSerializer = NewManifestSerializer.Get();
-
-					InManifestSerializer->SetLocalCaptureDir(AssetPath);
-
-					FManifestFileHeader Header(SubSequenceAssetName, SerializedType, FGuid());
-					FText Error;
-					FString FileName = FString::Printf(TEXT("%s_%s"), *(SerializedType.ToString()), *(SubSequenceAssetName));
-
-					if (!InManifestSerializer->OpenForWrite(FileName, Header, Error))
-					{
-						UE_LOG(SubSequenceSerialization, Warning, TEXT("Error Opening Sequence Sequencer File: Subject '%s' Error '%s'"), *(SubSequenceAssetName), *(Error.ToString()));
-					}
-				}
-			}
-
-			// Update our mappings of which sources use which sub-sequence.
-			SourceSubSequenceMap.FindOrAdd(Source) = TargetSequence;
-			Source->TimecodeSource = Timecode;
-			for (UTakeRecorderSource* NewlyAddedSource : Source->PreRecording(TargetSequence, TargetSequenceID, InMasterSequence, InManifestSerializer))
-			{
-				// Add it to our classes list of sources 
-				Sources.Add(NewlyAddedSource);
-
-				// And then track it separately so we can recursively call PreRecording 
-				NewSources.Add(NewlyAddedSource);
-			}
-
-			// We need to wait until PreRecording is called on a source before asking it to place itself in a folder
-			// so that the Source has had a chance to create any required sections that will go in the folder.
-			if (!Settings.bRecordSourcesIntoSubSequences && bCreateSequencerFolders)
-			{
-				UMovieSceneFolder* Folder = AddFolderForSource(Source, InMasterSequence->GetMovieScene());
-
-				// Different sources can create different kinds of tracks so we allow each source to decide how it gets
-				// represented inside the folder.
-				Source->AddContentsToFolder(Folder);
-			}
-		}
-	}
-
-	if (NewSources.Num())
-	{
-		// We don't want to nestle sub-sequences recursively so we always pass the Master Sequence and not the sequence
-		// created for a new source.
-		StartRecordingRecursive(NewSources, InMasterSequence, Timecode, InManifestSerializer);
-		SourcesSerialNumber++;
-
-		FQualifiedFrameTime QualifiedSequenceTime = GetCurrentRecordingFrameTime();
-		for (auto NewSource : NewSources)
-		{
-			if (NewSource->bEnabled)
-			{
-				ULevelSequence* SourceSequence = SourceSubSequenceMap[NewSource];
-				FFrameNumber FrameNumber = QualifiedSequenceTime.ConvertTo(SourceSequence->GetMovieScene()->GetTickResolution()).FloorToFrame();
-				NewSource->StartRecording(Timecode, FrameNumber, SourceSubSequenceMap[NewSource]);
-			}
-		}
-	}
-}
-
-
 void UTakeRecorderSources::PreRecordingRecursive(TArray<UTakeRecorderSource*> InSources, ULevelSequence* InMasterSequence, TArray<UTakeRecorderSource*>& NewSourcesOut, FManifestSerializer* InManifestSerializer)
 {
 
@@ -396,9 +233,9 @@ void UTakeRecorderSources::PreRecordingRecursive(TArray<UTakeRecorderSource*> In
 	}
 }
 
-void UTakeRecorderSources::StartRecordingPreRecordedSources(const FTimecode& CurrentTimecode)
+void UTakeRecorderSources::StartRecordingPreRecordedSources(const FQualifiedFrameTime& CurrentFrameTime)
 {
-	StartRecordingTheseSources(PreRecordedSources, CurrentTimecode);
+	StartRecordingTheseSources(PreRecordedSources, CurrentFrameTime);
 	PreRecordedSources.Reset(0);
 }
 
@@ -409,7 +246,7 @@ void UTakeRecorderSources::PreRecordSources(TArray<UTakeRecorderSource *> InSour
 }
 
 
-void UTakeRecorderSources::StartRecordingSource(TArray<UTakeRecorderSource *> InSources, const FTimecode& CurrentTimecode)
+void UTakeRecorderSources::StartRecordingSource(TArray<UTakeRecorderSource *> InSources, const FQualifiedFrameTime& CurrentFrameTime)
 {
 	// This calls PreRecording recursively on every source so that all sources that get added by another source
 	// have had PreRecording called.
@@ -419,17 +256,23 @@ void UTakeRecorderSources::StartRecordingSource(TArray<UTakeRecorderSource *> In
 	{
 		InSources.Append(NewSources);
 	}
-	StartRecordingTheseSources(InSources, CurrentTimecode);
+	StartRecordingTheseSources(InSources, CurrentFrameTime);
+}
+
+
+void UTakeRecorderSources::StartRecordingSource(TArray<UTakeRecorderSource *> InSources, const FTimecode& CurrentTimecode)
+{
+	StartRecordingSource(InSources, FQualifiedFrameTime(CurrentTimecode, TargetLevelSequenceDisplayRate)); 
 }
 
-void UTakeRecorderSources::StartRecordingTheseSources(const TArray<UTakeRecorderSource *>& InSources, const FTimecode& CurrentTimecode)
+void UTakeRecorderSources::StartRecordingTheseSources(const TArray<UTakeRecorderSource *>& InSources, const FQualifiedFrameTime& CurrentFrameTime)
 {
-	FQualifiedFrameTime QualifiedSequenceTime = GetCurrentRecordingFrameTime();
+	FTimecode CurrentTimecode = FTimecode::FromFrameNumber(FFrameRate::TransformTime(CurrentFrameTime.Time, TargetLevelSequenceTickResolution, TargetLevelSequenceDisplayRate).FloorToFrame(), TargetLevelSequenceDisplayRate);
+
 	for (auto Source : InSources)
 	{
 		if (Source->bEnabled)
 		{
-			ULevelSequence* SourceSequence = SourceSubSequenceMap[Source];
 			if (Settings.bRecordSourcesIntoSubSequences && Source->SupportsSubscenes()) //Set Timcode on MovieScene if we created a sub scene for it
 			{
 				for (UMovieSceneSubSection* ActiveSubSection : ActiveSubSections)
@@ -437,25 +280,23 @@ void UTakeRecorderSources::StartRecordingTheseSources(const TArray<UTakeRecorder
 					// Set timecode source and start time if it hasn't been set
 					if (ActiveSubSection->TimecodeSource.Timecode == FTimecode())
 					{
-						SetSectionStartTimecode(ActiveSubSection, CurrentTimecode, TargetLevelSequenceDisplayRate, TargetLevelSequenceTickResolution);
+						SetSectionStartTimecode(ActiveSubSection, CurrentTimecode, CurrentFrameTime);
 					}
 				}
 			}
-			FFrameNumber FrameNumber = QualifiedSequenceTime.ConvertTo(SourceSequence->GetMovieScene()->GetTickResolution()).FloorToFrame();
-			Source->TimecodeSource = CurrentTimecode;
 
-			Source->StartRecording(CurrentTimecode, FrameNumber, SourceSubSequenceMap[Source]);
+			Source->StartRecording(CurrentTimecode, CurrentFrameTime.Time.FloorToFrame(), SourceSubSequenceMap[Source]);
 		}
 	}
 }
 
-void UTakeRecorderSources::SetSectionStartTimecode(UMovieSceneSubSection* SubSection, const FTimecode& Timecode, FFrameRate FrameRate, FFrameRate TickResolution)
+void UTakeRecorderSources::SetSectionStartTimecode(UMovieSceneSubSection* SubSection, const FTimecode& Timecode, const FQualifiedFrameTime& CurrentFrameTime)
 {
-	FFrameNumber RecordStartFrame = Settings.bStartAtCurrentTimecode ? FFrameRate::TransformTime(FFrameTime(Timecode.ToFrameNumber(FrameRate)), FrameRate, TickResolution).FloorToFrame() : GetCurrentRecordingFrameTime().ConvertTo(TickResolution).FloorToFrame();
 	SubSection->TimecodeSource = FMovieSceneTimecodeSource(Timecode);
 
 	// Ensure we're expanded to at least the next frame so that we don't set the start past the end
 	// when we set the first frame.
+	FFrameNumber RecordStartFrame = CurrentFrameTime.Time.FloorToFrame(); 		
 	SubSection->ExpandToFrame(RecordStartFrame + FFrameNumber(1));
 	SubSection->SetStartFrame(TRangeBound<FFrameNumber>::Inclusive(RecordStartFrame));
 
@@ -473,42 +314,44 @@ void UTakeRecorderSources::SetCachedAssets(class ULevelSequence* InSequence, FMa
 	RecordedTimes.Empty();
 }
 
-void UTakeRecorderSources::PreRecording(class ULevelSequence* InSequence, FManifestSerializer* InManifestSerializer)
+void UTakeRecorderSources::PreRecording(class ULevelSequence* InSequence, const FQualifiedFrameTime& InCurrentFrameTime, FManifestSerializer* InManifestSerializer)
 {
+	CachedFrameTime = InCurrentFrameTime;
 	SetCachedAssets(InSequence, InManifestSerializer);
 	PreRecordSources(Sources);
 }
 
-void UTakeRecorderSources::StartRecording(class ULevelSequence* InSequence, const FTimecode& InTimecodeSource, FManifestSerializer* InManifestSerializer)
+void UTakeRecorderSources::StartRecording(class ULevelSequence* InSequence, const FQualifiedFrameTime& InCurrentFrameTime, FManifestSerializer* InManifestSerializer)
 {
+	CachedFrameTime = InCurrentFrameTime;
 	bIsRecording = true;
-	TimeSinceRecordingStarted = 0.f;
 	TargetLevelSequenceTickResolution = InSequence->GetMovieScene()->GetTickResolution();
 	TargetLevelSequenceDisplayRate = InSequence->GetMovieScene()->GetDisplayRate();
 
-	UE_LOG(LogTakesCore, Log, TEXT("StartRecording: %s"), *InTimecodeSource.ToString());
+	FTimecode CurrentTimecode = FTimecode::FromFrameNumber(FFrameRate::TransformTime(InCurrentFrameTime.Time, TargetLevelSequenceTickResolution, TargetLevelSequenceDisplayRate).FloorToFrame(), TargetLevelSequenceDisplayRate);
+
+	UE_LOG(LogTakesCore, Log, TEXT("StartRecording: %s"), *CurrentTimecode.ToString());
 
-	InSequence->GetMovieScene()->TimecodeSource = InTimecodeSource;
-	StartRecordingTimecodeSource = InTimecodeSource;
-	StartRecordingPreRecordedSources(InTimecodeSource);
+	InSequence->GetMovieScene()->TimecodeSource = CurrentTimecode;
+	StartRecordingPreRecordedSources(InCurrentFrameTime);
 }
 
-FFrameTime UTakeRecorderSources::AdvanceTime(float DeltaTime)
+FFrameTime UTakeRecorderSources::AdvanceTime(const FQualifiedFrameTime& CurrentFrameTime, float DeltaTime)
 {
-	FQualifiedFrameTime FrameTime = GetCurrentRecordingFrameTime();
-	//Time in seconds since recording started. Used when there is no Timecode Sync (e.g. in case it get's lost or dropped).
-	TimeSinceRecordingStarted += DeltaTime;
-	return FrameTime.ConvertTo(TargetLevelSequenceTickResolution);
+	CachedFrameTime = CurrentFrameTime;
+	return CurrentFrameTime.ConvertTo(TargetLevelSequenceTickResolution);
 }
 
-FFrameTime UTakeRecorderSources::TickRecording(class ULevelSequence* InSequence, const FTimecode& InTimecodeSource, float DeltaTime)
+FFrameTime UTakeRecorderSources::TickRecording(class ULevelSequence* InSequence, const FQualifiedFrameTime& CurrentFrameTime, float DeltaTime)
 {
+	CachedFrameTime = CurrentFrameTime;
+
+	FTimecode Timecode = FTimecode::FromFrameNumber(FFrameRate::TransformTime(CurrentFrameTime.Time, TargetLevelSequenceTickResolution, TargetLevelSequenceDisplayRate).FloorToFrame(), TargetLevelSequenceDisplayRate);
+
 #if !NO_LOGGING
-	UE_LOG(LogTakesCore, VeryVerbose, TEXT("TickRecording: %s"), *InTimecodeSource.ToString());
+	UE_LOG(LogTakesCore, VeryVerbose, TEXT("TickRecording: %s"), *Timecode.ToString());
 #endif
 
-	FQualifiedFrameTime FrameTime = GetCurrentRecordingFrameTime();
-	FQualifiedFrameTime SourceFrameTime(FrameTime);
 	bool bTimeIncremented = DeltaTime > 0.0f;
 
 	if (bTimeIncremented) //only record if time incremented, may not with timecode providers with low frame rates
@@ -517,50 +360,27 @@ FFrameTime UTakeRecorderSources::TickRecording(class ULevelSequence* InSequence,
 		{
 			if (Sources[SourceIndex]->bEnabled)
 			{
-				Sources[SourceIndex]->TickRecording(SourceFrameTime);
+				Sources[SourceIndex]->TickRecording(CurrentFrameTime);
 			}
 		}
 	}
 
 	{
-		const FFrameNumber StartFrameNumber = StartRecordingTimecodeSource.ToFrameNumber(TargetLevelSequenceDisplayRate);
-		FFrameTime StartTime = FFrameRate::TransformTime(FFrameTime(StartFrameNumber), TargetLevelSequenceDisplayRate, TargetLevelSequenceTickResolution);
-		const FFrameTime CurrentFrameTimeSinceStart = TargetLevelSequenceTickResolution.AsFrameTime(TimeSinceRecordingStarted);
-		FQualifiedFrameTime CurrentFrameTime = FQualifiedFrameTime(StartTime + CurrentFrameTimeSinceStart, TargetLevelSequenceTickResolution);
-
-		FFrameTime KeyTime = FFrameRate::TransformTime(CurrentFrameTime.Time, TargetLevelSequenceTickResolution, TargetLevelSequenceDisplayRate);
-		FTimecode Timecode = FTimecode::FromFrameNumber(KeyTime.FrameNumber, TargetLevelSequenceDisplayRate);
-		RecordedTimes.Add(TPair<FQualifiedFrameTime, FTimecode>(FrameTime, Timecode));
+		RecordedTimes.Add(TPair<FQualifiedFrameTime, FTimecode>(CurrentFrameTime, Timecode));
 	}
 
-	//Time in seconds since recording started. Used when there is no Timecode Sync (e.g. in case it get's lost or dropped).
-	TimeSinceRecordingStarted += DeltaTime;
-
-
 	// If we're recording into sub-sections we want to update their range every frame so they appear to
 	// animate as their contents are filled. We can't check against the size of all sections (not all
 	// source types have data in their sections until the end) and if you're partially re-recording
 	// a track it would size to the existing content which would skip the animation as well.
 
-	FFrameNumber EndFrame = Settings.bStartAtCurrentTimecode ? FFrameRate::TransformTime(FFrameTime(InTimecodeSource.ToFrameNumber(TargetLevelSequenceDisplayRate)), TargetLevelSequenceDisplayRate, TargetLevelSequenceTickResolution).CeilToFrame() : FrameTime.ConvertTo(TargetLevelSequenceTickResolution).CeilToFrame();
+	FFrameNumber EndFrame = CurrentFrameTime.Time.CeilToFrame();
 	for (UMovieSceneSubSection* SubSection : ActiveSubSections)
 	{
 		// Subsections will have been created to start at the time that they appeared, so we just need to expand their range to this recording time
 		SubSection->ExpandToFrame(EndFrame);
 	}
-	return FrameTime.ConvertTo(TargetLevelSequenceTickResolution);
-}
-
-//We now always just use TimeSinceRecordingStarted instead of possibly using timecode to determine our time since
-//That can give us a higher resolution
-FQualifiedFrameTime UTakeRecorderSources::GetCurrentRecordingFrameTime() const
-{
-	const FFrameNumber StartFrameNumber = StartRecordingTimecodeSource.ToFrameNumber(TargetLevelSequenceDisplayRate);
-	FFrameTime StartTime = Settings.bStartAtCurrentTimecode ? FFrameRate::TransformTime(FFrameTime(StartFrameNumber), TargetLevelSequenceDisplayRate, TargetLevelSequenceTickResolution) : CachedLevelSequence->GetMovieScene()->GetPlaybackRange().GetLowerBoundValue();
-
-	const FFrameTime CurrentFrameTimeSinceStart = TargetLevelSequenceTickResolution.AsFrameTime(TimeSinceRecordingStarted);
-	FQualifiedFrameTime FrameTime = FQualifiedFrameTime(StartTime + CurrentFrameTimeSinceStart, TargetLevelSequenceTickResolution);
-	return FrameTime;
+	return CurrentFrameTime.ConvertTo(TargetLevelSequenceTickResolution);
 }
 
 void UTakeRecorderSources::StopRecording(class ULevelSequence* InSequence)
@@ -568,7 +388,6 @@ void UTakeRecorderSources::StopRecording(class ULevelSequence* InSequence)
 	UE_LOG(LogTakesCore, Log, TEXT("StopRecording"));
 
 	bIsRecording = false;
-	TimeSinceRecordingStarted = 0.f;
 
 	for (auto Source : Sources)
 	{
diff --git a/Engine/Plugins/VirtualProduction/Takes/Source/TakesCore/Public/TakeRecorderSource.h b/Engine/Plugins/VirtualProduction/Takes/Source/TakesCore/Public/TakeRecorderSource.h
--- a/Engine/Plugins/VirtualProduction/Takes/Source/TakesCore/Public/TakeRecorderSource.h
+++ b/Engine/Plugins/VirtualProduction/Takes/Source/TakesCore/Public/TakeRecorderSource.h
@@ -43,8 +43,6 @@ public:
 	UPROPERTY(BlueprintReadWrite, Category="Source")
 	FColor TrackTint;
 
-	//Timecode source when recording is started via StartRecording.
-	FTimecode TimecodeSource;
 public:
 	/** 
 	* If the source is Valid and ready to be recorded. Tells the Take Recorder that it's really going to record. Note that PreRecord/StartRecord
diff --git a/Engine/Plugins/VirtualProduction/Takes/Source/TakesCore/Public/TakeRecorderSources.h b/Engine/Plugins/VirtualProduction/Takes/Source/TakesCore/Public/TakeRecorderSources.h
--- a/Engine/Plugins/VirtualProduction/Takes/Source/TakesCore/Public/TakeRecorderSources.h
+++ b/Engine/Plugins/VirtualProduction/Takes/Source/TakesCore/Public/TakeRecorderSources.h
@@ -100,7 +100,10 @@ public:
 
 	/** Calls the recording initialization flows on each of the specified sources. */
 	UFUNCTION(BlueprintCallable, Category = "Take Recorder")
-	void StartRecordingSource(TArray<UTakeRecorderSource*> InSources, const FTimecode& CurrentTiimecode);
+	void StartRecordingSource(TArray<UTakeRecorderSource*> InSources, const FQualifiedFrameTime& CurrentFrameTime);
+	
+	UE_DEPRECATED(4.27, "StartRecordingSource with FTimecode has been deprecated, please use StartRecordingSource with FQualifiedFrameTime")
+	void StartRecordingSource(TArray<UTakeRecorderSource*> InSources, const FTimecode& CurrentTimecode);
 
 public:
 
@@ -130,26 +133,26 @@ public:
 	 * Pre recording pass
 	*
 	*/
-	void PreRecording(class ULevelSequence* InSequence, FManifestSerializer* InManifestSerializer);
+	void PreRecording(class ULevelSequence* InSequence, const FQualifiedFrameTime& InCurrentFrameTime, FManifestSerializer* InManifestSerializer);
 
 	/*
 	 * Start recording pass
 	 *
 	 */
-	void StartRecording(class ULevelSequence* InSequence, const FTimecode& InTimecodeSource, FManifestSerializer* InManifestSerializer);
+	void StartRecording(class ULevelSequence* InSequence, const FQualifiedFrameTime& InCurrentFrameTime, FManifestSerializer* InManifestSerializer);
 
 	/*
 	 * Moves time forward by given DeltaTime
 	 *
 	 * @return Current Frame Number
 	 */
-	FFrameTime AdvanceTime(float DeltaTime);
+	FFrameTime AdvanceTime(const FQualifiedFrameTime& CurrentFrameTime, float DeltaTime);
 
 	/*
 	* Tick recording pass
 	* @return Current Frame Number we are recording at.
 	*/
-	FFrameTime TickRecording(class ULevelSequence* InSequence, const FTimecode& InTimecodeSource, float DeltaTime);
+	FFrameTime TickRecording(class ULevelSequence* InSequence, const FQualifiedFrameTime& CurrentFrameTime, float DeltaTime);
 
 	/*
 	* Stop recording pass
@@ -170,6 +173,8 @@ public:
 	 */
 	static TArray<TPair<FQualifiedFrameTime, FTimecode> > RecordedTimes;
 
+	FQualifiedFrameTime GetCachedFrameTime() const { return CachedFrameTime; }
+
 private:
 	/** Called at the end of each frame in both the Editor and in Game to update all Sources. */
 	virtual void Tick(float DeltaTime) {}
@@ -181,28 +186,22 @@ private:
 	DECLARE_MULTICAST_DELEGATE(FOnSourcesChanged);
 	FOnSourcesChanged OnSourcesChangedEvent;
 
-	/** Calls PreRecording on sources recursively allowing them to create other sources which properly get PreRecording called on them as well. */
-	void StartRecordingRecursive(TArray<UTakeRecorderSource*> InSources, ULevelSequence* InSequence, const FTimecode& Timecode, FManifestSerializer* InManifestSerializer);
-
 	/** Calls PreRecording on sources recursively allowing them to create other sources which properly get PreRecording called on them as well. */
 	void PreRecordingRecursive(TArray<UTakeRecorderSource*> InSources, ULevelSequence* InMasterSequence, TArray<UTakeRecorderSource*>& NewSourcesOut, FManifestSerializer* InManifestSerializer);
 
 	/** Finds the folder that the given Source should be created in, creating it if necessary. */
 	class UMovieSceneFolder* AddFolderForSource(const UTakeRecorderSource* InSource, class UMovieScene* InMovieScene);
 
-	/** Gets the current frame time for recording */
-	FQualifiedFrameTime GetCurrentRecordingFrameTime() const;
-
 	/** Remove object bindings that don't have any tracks and are not bindings for attach/path tracks */
 	void RemoveRedundantTracks();
 
-	void StartRecordingPreRecordedSources(const FTimecode& CurrentTimecode);
+	void StartRecordingPreRecordedSources(const FQualifiedFrameTime& CurrentFrameTime);
 
 	void PreRecordSources(TArray<UTakeRecorderSource *> InSources);
 
-	void StartRecordingTheseSources(const TArray<UTakeRecorderSource *>& InSources, const FTimecode& CurrentTimecode);
+	void StartRecordingTheseSources(const TArray<UTakeRecorderSource *>& InSources, const FQualifiedFrameTime& CurrentFrameTime);
 
-	void SetSectionStartTimecode(UMovieSceneSubSection* SubSection, const FTimecode& Timecode, FFrameRate FrameRate, FFrameRate TickResolution);
+	void SetSectionStartTimecode(UMovieSceneSubSection* SubSection, const FTimecode& Timecode, const FQualifiedFrameTime& CurrentFrameTime);
 
 private:
 
@@ -221,9 +220,6 @@ private:
 	/** Are we currently in a recording pass and should be ticking our Sources? */
 	bool bIsRecording;
 
-	/** Float of how long the recording has been going based on delta tick times. Used when we have no Timecode Synchronization */
-	float TimeSinceRecordingStarted;
-
 	/** What Tick Resolution is the target level sequence we're recording into? Used to convert seconds into FrameNumbers. */
 	FFrameRate TargetLevelSequenceTickResolution;
 
@@ -245,10 +241,9 @@ private:
 	/** Array of Allocated Serializers created for each sub sequence.  Deleted at the end of the recording so memory is freed. */
 	TArray<TSharedPtr<FManifestSerializer>> CreatedManifestSerializers;
 
-	/** Timecode time at start of recording */
-	FTimecode StartRecordingTimecodeSource;
-
 	/** All sources after PreRecord */
 	TArray<UTakeRecorderSource *> PreRecordedSources;
 
+	/** The last frame time during tick recording */
+	FQualifiedFrameTime CachedFrameTime;
 };
