Take Recorder: Static analysis fix

#rnx

diff --git a/Engine/Plugins/VirtualProduction/Takes/Source/TakeRecorder/Private/Recorder/TakeRecorder.cpp b/Engine/Plugins/VirtualProduction/Takes/Source/TakeRecorder/Private/Recorder/TakeRecorder.cpp
--- a/Engine/Plugins/VirtualProduction/Takes/Source/TakeRecorder/Private/Recorder/TakeRecorder.cpp
+++ b/Engine/Plugins/VirtualProduction/Takes/Source/TakeRecorder/Private/Recorder/TakeRecorder.cpp
@@ -1067,7 +1067,10 @@ void UTakeRecorder::Stop()
 		// Restore the playback range to what it was before recording.
 		if (Parameters.TakeRecorderMode == ETakeRecorderMode::RecordIntoSequence)
 		{
-			MovieScene->SetPlaybackRange(CachedPlaybackRange);
+			if (MovieScene)
+			{
+				MovieScene->SetPlaybackRange(CachedPlaybackRange);
+			}
 		}
 		else
 		{
