Take Recorder: Fix multiuser when there is no sequence editor to get the current time for

diff --git a/Engine/Plugins/VirtualProduction/Takes/Source/TakeRecorder/Private/Recorder/TakeRecorder.cpp b/Engine/Plugins/VirtualProduction/Takes/Source/TakeRecorder/Private/Recorder/TakeRecorder.cpp
--- a/Engine/Plugins/VirtualProduction/Takes/Source/TakeRecorder/Private/Recorder/TakeRecorder.cpp
+++ b/Engine/Plugins/VirtualProduction/Takes/Source/TakeRecorder/Private/Recorder/TakeRecorder.cpp
@@ -788,19 +788,55 @@ void UTakeRecorder::Tick(float DeltaTime)
 	}
 }
 
+FQualifiedFrameTime UTakeRecorder::GetRecordTime() const
+{
+	FQualifiedFrameTime RecordTime;
+
+	TSharedPtr<ISequencer> Sequencer = WeakSequencer.Pin();
+	if (Sequencer.IsValid())
+	{
+		RecordTime = Sequencer->GetGlobalTime();
+	}
+	else if (SequenceAsset)
+	{
+		UMovieScene* MovieScene = SequenceAsset->GetMovieScene();
+		if (MovieScene)
+		{
+			FFrameRate FrameRate = MovieScene->GetDisplayRate();
+			FFrameRate TickResolution = MovieScene->GetTickResolution();
+
+			FTimecode CurrentTimecode = FApp::GetTimecode();
+		
+			FFrameNumber CurrentFrame = FFrameRate::TransformTime(FFrameTime(CurrentTimecode.ToFrameNumber(FrameRate)), FrameRate, TickResolution).FloorToFrame();
+			FFrameNumber FrameAtStart = FFrameRate::TransformTime(FFrameTime(TimecodeAtStart.ToFrameNumber(FrameRate)), FrameRate, TickResolution).FloorToFrame();
+
+			if (Parameters.Project.bStartAtCurrentTimecode)
+			{
+				RecordTime = FQualifiedFrameTime(CurrentFrame, TickResolution);
+			}
+			else
+			{
+				RecordTime = FQualifiedFrameTime(CurrentFrame - FrameAtStart, TickResolution);
+			}
+		}
+	}
+
+	return RecordTime;
+}
+
 void UTakeRecorder::InternalTick(float DeltaTime)
 {
 	TSharedPtr<ISequencer> Sequencer = WeakSequencer.Pin();
-	FQualifiedFrameTime SequencerTime = Sequencer.IsValid() ? Sequencer->GetGlobalTime() : FQualifiedFrameTime();
-	
+	FQualifiedFrameTime RecordTime = GetRecordTime();
+		
 	UTakeRecorderSources* Sources = SequenceAsset->FindOrAddMetaData<UTakeRecorderSources>();
 	if (!Parameters.bDisableRecordingAndSave)
 	{
-		CurrentFrameTime = Sources->TickRecording(SequenceAsset, SequencerTime, DeltaTime);
+		CurrentFrameTime = Sources->TickRecording(SequenceAsset, RecordTime, DeltaTime);
 	}
 	else
 	{
-		CurrentFrameTime = Sources->AdvanceTime(SequencerTime, DeltaTime);
+		CurrentFrameTime = Sources->AdvanceTime(RecordTime, DeltaTime);
 	}
 
 	if (Sequencer.IsValid())
@@ -906,7 +942,8 @@ void UTakeRecorder::Start()
 	TSharedPtr<ISequencer> Sequencer = WeakSequencer.Pin();
 
 	CurrentFrameTime = FFrameTime(0);
-		
+	TimecodeAtStart = Timecode;
+
 	// Discard any entity tokens we have so that restore state does not take effect when we delete any sections that recording will be replacing.
 	if (Sequencer.IsValid())
 	{
@@ -961,12 +998,12 @@ void UTakeRecorder::Start()
 
 	if (!Parameters.bDisableRecordingAndSave)
 	{
-		FQualifiedFrameTime SequencerTime = Sequencer.IsValid() ? Sequencer->GetGlobalTime() : FQualifiedFrameTime();
+		FQualifiedFrameTime RecordTime = GetRecordTime();
 
-		Sources->StartRecording(SequenceAsset, SequencerTime, Parameters.User.bAutoSerialize ? &ManifestSerializer : nullptr);
+		Sources->StartRecording(SequenceAsset, RecordTime, Parameters.User.bAutoSerialize ? &ManifestSerializer : nullptr);
 
 		// Record immediately so that there's a key on the first frame of recording
-		Sources->TickRecording(SequenceAsset, SequencerTime, 0.1f);
+		Sources->TickRecording(SequenceAsset, RecordTime, 0.1f);
 	}
 
 	if (!ShouldShowNotifications())
diff --git a/Engine/Plugins/VirtualProduction/Takes/Source/TakeRecorder/Public/Recorder/TakeRecorder.h b/Engine/Plugins/VirtualProduction/Takes/Source/TakeRecorder/Public/Recorder/TakeRecorder.h
--- a/Engine/Plugins/VirtualProduction/Takes/Source/TakeRecorder/Public/Recorder/TakeRecorder.h
+++ b/Engine/Plugins/VirtualProduction/Takes/Source/TakeRecorder/Public/Recorder/TakeRecorder.h
@@ -204,6 +204,9 @@ private:
 
 private:
 
+	/* The time at which to record. Taken from the Sequencer global time, otherwise based on timecode */
+	FQualifiedFrameTime GetRecordTime() const;
+
 	/** Called by Tick and Start to make sure we record at start */
 	void InternalTick(float DeltaTime);
 
@@ -220,6 +223,9 @@ private:
 	/** FFrameTime in MovieScene Resolution we are at*/
 	FFrameTime CurrentFrameTime;
 
+	/** Timecode at the start of recording */
+	FTimecode TimecodeAtStart;
+
 	/** Optional frame to stop recording at*/
 	TOptional<FFrameNumber> StopRecordingFrame;
 
