Sequencer: Large world coordinates support

- Add support for double precision channels, curves, evaluation, blending, and all other runtime infrastructure.
- Note that, as usual for now, double channels load and save float values.
- Editor side also gains some new track editors for these new types, with some workarounds to correctly recognize between float and double vectors.
- Transform tracks are now operating entirely in doubles.
- Float recomposing APIs for keying tracks in the editor are now using doubles, and have been renamed to "value recomposing".

diff --git a/Engine/Config/BaseEngine.ini b/Engine/Config/BaseEngine.ini
--- a/Engine/Config/BaseEngine.ini
+++ b/Engine/Config/BaseEngine.ini
@@ -919,6 +919,12 @@ NearClipPlane=10.0
 +PropertyRedirects=(OldName="GeometryCollectionComponent.Simulating", NewName="SimulatePhysics")
 +PropertyRedirects=(OldName="ControlRigComponent.OnPreUpdateDelegate", NewName="OnPreForwardsSolveDelegate")
 +PropertyRedirects=(OldName="ControlRigComponent.OnPostUpdateDelegate", NewName="OnPostForwardsSolveDelegate")
+; Sequencer LWC support
++ClassRedirects=(OldName="/Script/MovieSceneTracks.MovieSceneVectorSection",NewName="/Script/MovieSceneTracks.MovieSceneFloatVectorSection")
++ClassRedirects=(OldName="/Script/MovieSceneTracks.MovieSceneVectorTrack",NewName="/Script/MovieSceneTracks.MovieSceneFloatVectorTrack")
++ClassRedirects=(OldName="/Script/MovieSceneTracks.MovieSceneVectorTrackExtensions",NewName="/Script/SequencerScripting.MovieSceneFloatVectorTrackExtensions")
++StructRedirects=(OldName="/Script/MovieSceneTracks.MovieSceneVectorKeyStructBase",NewName="/Script/MovieSceneTracks.MovieSceneFloatVectorKeyStructBase")
++StructRedirects=(OldName="/Script/MovieSceneTracks.MovieSceneVectorKeyStruct",NewName="/Script/MovieSceneTracks.MovieSceneVector3fKeyStruct")
 
 
 [CoreUObject.Metadata]
diff --git a/Engine/Plugins/Animation/ControlRig/Source/ControlRig/Private/Sequencer/MovieSceneControlRigParameterTrack.cpp b/Engine/Plugins/Animation/ControlRig/Source/ControlRig/Private/Sequencer/MovieSceneControlRigParameterTrack.cpp
--- a/Engine/Plugins/Animation/ControlRig/Source/ControlRig/Private/Sequencer/MovieSceneControlRigParameterTrack.cpp
+++ b/Engine/Plugins/Animation/ControlRig/Source/ControlRig/Private/Sequencer/MovieSceneControlRigParameterTrack.cpp
@@ -409,6 +409,7 @@ TArray<FFBXNodeAndChannels>* UMovieSceneControlRigParameterTrack::GetNodeAndChan
 		return nullptr;
 	}
 
+	const FName DoubleChannelTypeName = FMovieSceneDoubleChannel::StaticStruct()->GetFName();
 	const FName FloatChannelTypeName = FMovieSceneFloatChannel::StaticStruct()->GetFName();
 	const FName BoolChannelTypeName = FMovieSceneBoolChannel::StaticStruct()->GetFName();
 	const FName EnumChannelTypeName = FMovieSceneByteChannel::StaticStruct()->GetFName();
@@ -422,7 +423,7 @@ TArray<FFBXNodeAndChannels>* UMovieSceneControlRigParameterTrack::GetNodeAndChan
 	for (const FMovieSceneChannelEntry& Entry : CurrentSectionToKey->GetChannelProxy().GetAllEntries())
 	{
 		const FName ChannelTypeName = Entry.GetChannelTypeName();
-		if (ChannelTypeName != FloatChannelTypeName && ChannelTypeName != BoolChannelTypeName
+		if (ChannelTypeName != DoubleChannelTypeName && ChannelTypeName != FloatChannelTypeName && ChannelTypeName != BoolChannelTypeName
 			&& ChannelTypeName != EnumChannelTypeName && ChannelTypeName != IntegerChannelTypeName)
 		{
 			continue;
@@ -454,10 +455,15 @@ TArray<FFBXNodeAndChannels>* UMovieSceneControlRigParameterTrack::GetNodeAndChan
 						NodeAndChannel.NodeName = NodeName;
 						NodeAndChannels->Add(NodeAndChannel);
 					}
-					if (ChannelTypeName == FloatChannelTypeName)
+					if (ChannelTypeName == DoubleChannelTypeName)
+					{
+						FMovieSceneDoubleChannel* DoubleChannel = Channel.Cast<FMovieSceneDoubleChannel>().Get();
+						(*NodeAndChannels)[NodeAndChannels->Num() - 1].DoubleChannels.Add(DoubleChannel);
+					}
+					else if (ChannelTypeName == FloatChannelTypeName)
 					{
 						FMovieSceneFloatChannel* FloatChannel = Channel.Cast<FMovieSceneFloatChannel>().Get();
-						(*NodeAndChannels)[NodeAndChannels->Num() - 1].Channels.Add(FloatChannel);
+						(*NodeAndChannels)[NodeAndChannels->Num() - 1].FloatChannels.Add(FloatChannel);
 					}
 					else if (ChannelTypeName == BoolChannelTypeName)
 					{
diff --git a/Engine/Plugins/Animation/ControlRig/Source/ControlRigEditor/Private/Tools/ControlRigSnapper.cpp b/Engine/Plugins/Animation/ControlRig/Source/ControlRigEditor/Private/Tools/ControlRigSnapper.cpp
--- a/Engine/Plugins/Animation/ControlRig/Source/ControlRigEditor/Private/Tools/ControlRigSnapper.cpp
+++ b/Engine/Plugins/Animation/ControlRig/Source/ControlRigEditor/Private/Tools/ControlRigSnapper.cpp
@@ -236,7 +236,7 @@ struct FGuidAndActor
 			}
 		}
 
-		TArrayView<FMovieSceneFloatChannel*> Channels = TransformSection->GetChannelProxy().GetChannels<FMovieSceneFloatChannel>();
+		TArrayView<FMovieSceneDoubleChannel*> Channels = TransformSection->GetChannelProxy().GetChannels<FMovieSceneDoubleChannel>();
 		for (int32 Index = 0; Index < Frames.Num(); ++Index)
 		{
 			const FFrameNumber& Frame = Frames[Index];
@@ -297,7 +297,7 @@ struct FGuidAndActor
 
 			if (SnapSettings->bSnapPosition)
 			{
-				TMovieSceneChannelData<FMovieSceneFloatValue> ChannelData = Channels[0]->GetData();
+				TMovieSceneChannelData<FMovieSceneDoubleValue> ChannelData = Channels[0]->GetData();
 				MovieSceneToolHelpers::SetOrAddKey(ChannelData, Frame, Location.X);
 				ChannelData = Channels[1]->GetData();
 				MovieSceneToolHelpers::SetOrAddKey(ChannelData, Frame, Location.Y);
@@ -307,7 +307,7 @@ struct FGuidAndActor
 			if (SnapSettings->bSnapRotation)
 			{
 				//todo winding
-				TMovieSceneChannelData<FMovieSceneFloatValue> ChannelData = Channels[3]->GetData();
+				TMovieSceneChannelData<FMovieSceneDoubleValue> ChannelData = Channels[3]->GetData();
 				MovieSceneToolHelpers::SetOrAddKey(ChannelData, Frame, Rotation.Yaw);
 				ChannelData = Channels[4]->GetData();
 				MovieSceneToolHelpers::SetOrAddKey(ChannelData, Frame, Rotation.Pitch);
@@ -316,7 +316,7 @@ struct FGuidAndActor
 			}
 			if (SnapSettings->bSnapScale)
 			{
-				TMovieSceneChannelData<FMovieSceneFloatValue> ChannelData = Channels[6]->GetData();
+				TMovieSceneChannelData<FMovieSceneDoubleValue> ChannelData = Channels[6]->GetData();
 				MovieSceneToolHelpers::SetOrAddKey(ChannelData, Frame, Scale3D.X);
 				ChannelData = Channels[7]->GetData();
 				MovieSceneToolHelpers::SetOrAddKey(ChannelData, Frame, Scale3D.X);
diff --git a/Engine/Plugins/Editor/SequencerAnimTools/Source/SequencerAnimTools/Private/Sequencer/MotionTrailMovieSceneKey.cpp b/Engine/Plugins/Editor/SequencerAnimTools/Source/SequencerAnimTools/Private/Sequencer/MotionTrailMovieSceneKey.cpp
--- a/Engine/Plugins/Editor/SequencerAnimTools/Source/SequencerAnimTools/Private/Sequencer/MotionTrailMovieSceneKey.cpp
+++ b/Engine/Plugins/Editor/SequencerAnimTools/Source/SequencerAnimTools/Private/Sequencer/MotionTrailMovieSceneKey.cpp
@@ -5,7 +5,7 @@
 #include "TrailHierarchy.h"
 #include "Tracks/MovieScene3DTransformTrack.h"
 #include "Sections/MovieScene3DTransformSection.h"
-#include "Channels/MovieSceneFloatChannel.h"
+#include "Channels/MovieSceneDoubleChannel.h"
 #include "Channels/MovieSceneChannelProxy.h"
 #include "EntitySystem/Interrogation/MovieSceneInterrogationLinker.h"
 #include "HitProxies.h"
@@ -194,14 +194,14 @@ void FMotionTraiMovieScenelKeyTool::BuildKeys()
 	CachedSelection.Reset();
 
 	UMovieSceneSection* AbsoluteTransformSection = OwningTrail->GetSection();
-	TArrayView<FMovieSceneFloatChannel*> FloatChannels = AbsoluteTransformSection->GetChannelProxy().GetChannels<FMovieSceneFloatChannel>();
-	FloatChannels = FloatChannels.Slice(OwningTrail->GetChannelOffset(), uint8(ETransformChannel::MaxChannel) + 1);
+	TArrayView<FMovieSceneDoubleChannel*> DoubleChannels = AbsoluteTransformSection->GetChannelProxy().GetChannels<FMovieSceneDoubleChannel>();
+	DoubleChannels = DoubleChannels.Slice(OwningTrail->GetChannelOffset(), uint8(ETransformChannel::MaxChannel) + 1);
 	for (uint8 ChannelIdx = 0; ChannelIdx <= uint8(ETransformChannel::MaxChannel); ChannelIdx++)
 	{
-		FMovieSceneFloatChannel* FloatChannel = FloatChannels[ChannelIdx];
-		for (int32 Idx = 0; Idx < FloatChannel->GetNumKeys(); Idx++)
+		FMovieSceneDoubleChannel* DoubleChannel = DoubleChannels[ChannelIdx];
+		for (int32 Idx = 0; Idx < DoubleChannel->GetNumKeys(); Idx++)
 		{
-			FFrameNumber CurTime = FloatChannel->GetTimes()[Idx];
+			FFrameNumber CurTime = DoubleChannel->GetTimes()[Idx];
 
 			if (!Keys.Contains(CurTime))
 			{
@@ -242,17 +242,17 @@ void FMotionTraiMovieScenelKeyTool::TranslateSelectedKeys(bool bRight)
 		int32 Shift = bRight ? 1 : -1;
 		FFrameNumber Delta = FQualifiedFrameTime(Shift, OwningTrail->GetSequencer()->GetFocusedDisplayRate()).ConvertTo(OwningTrail->GetSequencer()->GetFocusedTickResolution()).RoundToFrame();
 		TRange<FFrameNumber> SectionNewBounds = Section->GetRange();
-		TArrayView<FMovieSceneFloatChannel*> FloatChannels = Section->GetChannelProxy().GetChannels<FMovieSceneFloatChannel>();
-		FloatChannels = FloatChannels.Slice(OwningTrail->GetChannelOffset(), uint8(ETransformChannel::MaxChannel) + 1);
+		TArrayView<FMovieSceneDoubleChannel*> DoubleChannels = Section->GetChannelProxy().GetChannels<FMovieSceneDoubleChannel>();
+		DoubleChannels = DoubleChannels.Slice(OwningTrail->GetChannelOffset(), uint8(ETransformChannel::MaxChannel) + 1);
 		for (uint8 ChannelIdx = 0; ChannelIdx <= uint8(ETransformChannel::MaxChannel); ChannelIdx++)
 		{
 			TArray<FFrameNumber> KeyTimes;
-			FMovieSceneFloatChannel* FloatChannel = FloatChannels[ChannelIdx];
-			TArray<FKeyHandle> KeyHandles = GetSelectedKeyHandles(FloatChannel);
+			FMovieSceneDoubleChannel* DoubleChannel = DoubleChannels[ChannelIdx];
+			TArray<FKeyHandle> KeyHandles = GetSelectedKeyHandles(DoubleChannel);
 			if (KeyHandles.Num() > 0)
 			{
 				KeyTimes.SetNum(KeyHandles.Num());
-				FloatChannel->GetKeyTimes(KeyHandles, KeyTimes);
+				DoubleChannel->GetKeyTimes(KeyHandles, KeyTimes);
 				FFrameNumber LowestFrameTime = KeyTimes[0];
 				FFrameNumber HighestFrameTime = KeyTimes[0];
 
@@ -272,7 +272,7 @@ void FMotionTraiMovieScenelKeyTool::TranslateSelectedKeys(bool bRight)
 				}
 				SectionNewBounds = TRange<FFrameNumber>::Hull(SectionNewBounds, TRange<FFrameNumber>(LowestFrameTime, HighestFrameTime + 1));
 				Section->SetRange(SectionNewBounds);
-				FloatChannel->SetKeyTimes(KeyHandles, KeyTimes);
+				DoubleChannel->SetKeyTimes(KeyHandles, KeyTimes);
 			}
 
 		}
@@ -294,13 +294,13 @@ void FMotionTraiMovieScenelKeyTool::DeleteSelectedKeys()
 	if (CachedSelection.Num() > 0 && Section && Section->TryModify())
 	{
 		Section->Modify();
-		TArrayView<FMovieSceneFloatChannel*> FloatChannels = Section->GetChannelProxy().GetChannels<FMovieSceneFloatChannel>();
-		FloatChannels = FloatChannels.Slice(OwningTrail->GetChannelOffset(), uint8(ETransformChannel::MaxChannel) + 1);
+		TArrayView<FMovieSceneDoubleChannel*> DoubleChannels = Section->GetChannelProxy().GetChannels<FMovieSceneDoubleChannel>();
+		DoubleChannels = DoubleChannels.Slice(OwningTrail->GetChannelOffset(), uint8(ETransformChannel::MaxChannel) + 1);
 		for (uint8 ChannelIdx = 0; ChannelIdx <= uint8(ETransformChannel::MaxChannel); ChannelIdx++)
 		{
-			FMovieSceneFloatChannel* FloatChannel = FloatChannels[ChannelIdx];
-			TArray<FKeyHandle> KeyHandles = GetSelectedKeyHandles(FloatChannel);
-			FloatChannel->DeleteKeys(KeyHandles);
+			FMovieSceneDoubleChannel* DoubleChannel = DoubleChannels[ChannelIdx];
+			TArray<FKeyHandle> KeyHandles = GetSelectedKeyHandles(DoubleChannel);
+			DoubleChannel->DeleteKeys(KeyHandles);
 		}
 		OwningTrail->GetSequencer()->NotifyMovieSceneDataChanged(EMovieSceneDataChangeType::TrackValueChanged);
 		CachedSelection.Reset();
@@ -432,15 +432,15 @@ TArray<FFrameNumber> FMotionTraiMovieScenelKeyTool::GetTimesFromModifiedTimes(co
 bool FMotionTraiMovieScenelKeyTool::ShouldRebuildKeys()
 {
 	TMap<FFrameNumber, TSet<ETransformChannel>> KeyTimes;
-	TArrayView<FMovieSceneFloatChannel*> FloatChannels = OwningTrail->GetSection()->GetChannelProxy().GetChannels<FMovieSceneFloatChannel>();
-	FloatChannels = FloatChannels.Slice(OwningTrail->GetChannelOffset(), uint8(ETransformChannel::MaxChannel) + 1);
+	TArrayView<FMovieSceneDoubleChannel*> DoubleChannels = OwningTrail->GetSection()->GetChannelProxy().GetChannels<FMovieSceneDoubleChannel>();
+	DoubleChannels = DoubleChannels.Slice(OwningTrail->GetChannelOffset(), uint8(ETransformChannel::MaxChannel) + 1);
 	for (uint8 ChannelIdx = 0; ChannelIdx <= uint8(ETransformChannel::MaxChannel); ChannelIdx++)
 
 	{
-		FMovieSceneFloatChannel* FloatChannel = FloatChannels[ChannelIdx];
-		for (int32 Idx = 0; Idx < FloatChannel->GetNumKeys(); Idx++)
+		FMovieSceneDoubleChannel* DoubleChannel = DoubleChannels[ChannelIdx];
+		for (int32 Idx = 0; Idx < DoubleChannel->GetNumKeys(); Idx++)
 		{
-			const FFrameNumber CurTime = FloatChannel->GetTimes()[Idx];
+			const FFrameNumber CurTime = DoubleChannel->GetTimes()[Idx];
 			KeyTimes.FindOrAdd(CurTime).Add(ETransformChannel(ChannelIdx));
 		}
 	}
@@ -460,7 +460,7 @@ bool FMotionTraiMovieScenelKeyTool::ShouldRebuildKeys()
 		{
 			const ETransformChannel TransformChannel = ETransformChannel(ChannelIdx);
 			if ((!TimeKeyPair.Value.Contains(TransformChannel) && Keys[TimeKeyPair.Key]->IdxMap.Contains(TransformChannel)) ||
-				(Keys[TimeKeyPair.Key]->IdxMap.Contains(TransformChannel) && FloatChannels[uint8(TransformChannel)]->GetData().GetIndex(Keys[TimeKeyPair.Key]->IdxMap[TransformChannel]) == INDEX_NONE) ||
+				(Keys[TimeKeyPair.Key]->IdxMap.Contains(TransformChannel) && DoubleChannels[uint8(TransformChannel)]->GetData().GetIndex(Keys[TimeKeyPair.Key]->IdxMap[TransformChannel]) == INDEX_NONE) ||
 				(TimeKeyPair.Value.Contains(TransformChannel) && !Keys[TimeKeyPair.Key]->IdxMap.Contains(TransformChannel)))
 			{
 				return true;
@@ -501,7 +501,7 @@ FTrailKeyInfo::FTrailKeyInfo(const FFrameNumber InFrameNumber, UMovieSceneSectio
 , bDirty(true)
 , OwningTrail(InOwningTrail)
 {
-	TArrayView<FMovieSceneFloatChannel*> Channels = InSection->GetChannelProxy().GetChannels<FMovieSceneFloatChannel>();
+	TArrayView<FMovieSceneDoubleChannel*> Channels = InSection->GetChannelProxy().GetChannels<FMovieSceneDoubleChannel>();
 	Channels = Channels.Slice(OwningTrail->GetChannelOffset(), uint8(ETransformChannel::MaxChannel) + 1);
 	for (uint8 Idx = 0; Idx <= uint8(ETransformChannel::MaxChannel); Idx++)
 	{
@@ -524,4 +524,4 @@ void FTrailKeyInfo::UpdateKeyTransform(EGetKeyFrom UpdateType)
 }
 
 } // namespace MovieScene
-} // namespace UE
\ No newline at end of file
+} // namespace UE
diff --git a/Engine/Plugins/Editor/SequencerAnimTools/Source/SequencerAnimTools/Private/Sequencer/MovieSceneTransformTrail.cpp b/Engine/Plugins/Editor/SequencerAnimTools/Source/SequencerAnimTools/Private/Sequencer/MovieSceneTransformTrail.cpp
--- a/Engine/Plugins/Editor/SequencerAnimTools/Source/SequencerAnimTools/Private/Sequencer/MovieSceneTransformTrail.cpp
+++ b/Engine/Plugins/Editor/SequencerAnimTools/Source/SequencerAnimTools/Private/Sequencer/MovieSceneTransformTrail.cpp
@@ -6,7 +6,7 @@
 #include "ISequencer.h"
 #include "Tracks/MovieScene3DTransformTrack.h"
 #include "Sections/MovieScene3DTransformSection.h"
-#include "Channels/MovieSceneFloatChannel.h"
+#include "Channels/MovieSceneDoubleChannel.h"
 #include "Channels/MovieSceneChannelProxy.h"
 #include "Systems/MovieSceneComponentTransformSystem.h"
 #include "MovieSceneSequence.h"
@@ -283,10 +283,10 @@ TRange<double> FMovieSceneTransformTrail::GetEffectiveSectionRange(int32 Channel
 
 	TRange<FFrameNumber> EffectiveRange = TRange<FFrameNumber>::Empty();
 
-	TArrayView<FMovieSceneFloatChannel*> Channels = TransformSection->GetChannelProxy().GetChannels<FMovieSceneFloatChannel>();
+	TArrayView<FMovieSceneDoubleChannel*> Channels = TransformSection->GetChannelProxy().GetChannels<FMovieSceneDoubleChannel>();
 	for (int32 ChannelIdx = ChannelOffset; ChannelIdx <= ChannelOffset + uint8(ETransformChannel::MaxChannel); ChannelIdx++)
 	{
-		FMovieSceneFloatChannel* Channel = Channels[ChannelIdx];
+		FMovieSceneDoubleChannel* Channel = Channels[ChannelIdx];
 		EffectiveRange = TRange<FFrameNumber>::Hull(EffectiveRange, Channel->ComputeEffectiveRange());
 	}
 
@@ -411,7 +411,7 @@ static bool SetActorTransform(ISequencer* Sequencer, AActor *Actor, UMovieScene3
 		}
 	}
 
-	TArrayView<FMovieSceneFloatChannel*> Channels = TransformSection->GetChannelProxy().GetChannels<FMovieSceneFloatChannel>();
+	TArrayView<FMovieSceneDoubleChannel*> Channels = TransformSection->GetChannelProxy().GetChannels<FMovieSceneDoubleChannel>();
 	for (int32 Index = 0; Index < Frames.Num(); ++Index)
 	{
 		const FFrameNumber& Frame = Frames[Index];
@@ -421,7 +421,7 @@ static bool SetActorTransform(ISequencer* Sequencer, AActor *Actor, UMovieScene3
 		FVector Location = LocalTransform.GetLocation();
 		FRotator Rotation = LocalTransform.GetRotation().Rotator();
 			
-		TMovieSceneChannelData<FMovieSceneFloatValue> ChannelData = Channels[0]->GetData();
+		TMovieSceneChannelData<FMovieSceneDoubleValue> ChannelData = Channels[0]->GetData();
 		MovieSceneToolHelpers::SetOrAddKey(ChannelData, Frame, Location.X);
 		ChannelData = Channels[1]->GetData();
 		MovieSceneToolHelpers::SetOrAddKey(ChannelData, Frame, Location.Y);
@@ -782,4 +782,4 @@ bool FMovieSceneControlRigTransformTrail::HandleAltClick(FEditorViewportClient*
 } // namespace MovieScene
 } // namespace UE
 
-#undef LOCTEXT_NAMESPACE
\ No newline at end of file
+#undef LOCTEXT_NAMESPACE
diff --git a/Engine/Plugins/Enterprise/DatasmithImporter/Source/DatasmithImporter/Private/DatasmithLevelSequenceImporter.cpp b/Engine/Plugins/Enterprise/DatasmithImporter/Source/DatasmithImporter/Private/DatasmithLevelSequenceImporter.cpp
--- a/Engine/Plugins/Enterprise/DatasmithImporter/Source/DatasmithImporter/Private/DatasmithLevelSequenceImporter.cpp
+++ b/Engine/Plugins/Enterprise/DatasmithImporter/Source/DatasmithImporter/Private/DatasmithLevelSequenceImporter.cpp
@@ -11,7 +11,7 @@
 #include "Algo/Transform.h"
 #include "Channels/MovieSceneBoolChannel.h"
 #include "Channels/MovieSceneChannelProxy.h"
-#include "Channels/MovieSceneFloatChannel.h"
+#include "Channels/MovieSceneDoubleChannel.h"
 #include "LevelSequence.h"
 #include "MovieScene.h"
 #include "MovieSceneSection.h"
@@ -137,7 +137,7 @@ namespace DatasmithLevelSequenceImporterImpl
 		}
 	}
 
-	int32 GetTransformValues(const TSharedRef<IDatasmithTransformAnimationElement>& Animation, EDatasmithTransformType TransformType, const FFrameRate& SrcFrameRate, const FFrameRate& DestFrameRate, TArray<FFrameNumber>& FrameNumbers, TArray<float>& XValues, TArray<float>& YValues, TArray<float>& ZValues, FFrameNumber& MinFrameNumber, FFrameNumber& MaxFrameNumber)
+	int32 GetTransformValues(const TSharedRef<IDatasmithTransformAnimationElement>& Animation, EDatasmithTransformType TransformType, const FFrameRate& SrcFrameRate, const FFrameRate& DestFrameRate, TArray<FFrameNumber>& FrameNumbers, TArray<double>& XValues, TArray<double>& YValues, TArray<double>& ZValues, FFrameNumber& MinFrameNumber, FFrameNumber& MaxFrameNumber)
 	{
 		int32 NumFrames = Animation->GetFramesCount(TransformType);
 		FrameNumbers.Reset(NumFrames);
@@ -209,7 +209,7 @@ namespace DatasmithLevelSequenceImporterImpl
 		return FrameNumbers.Num();
 	}
 
-	void SetTransformChannels(TArrayView<FMovieSceneFloatChannel*>& Channels, EDatasmithTransformType TransformType, ERichCurveInterpMode Interpolation, const TArray<FFrameNumber>& FrameNumbers, TArray<float>& XValues, TArray<float>& YValues, TArray<float>& ZValues)
+	void SetTransformChannels(TArrayView<FMovieSceneDoubleChannel*>& Channels, EDatasmithTransformType TransformType, ERichCurveInterpMode Interpolation, const TArray<FFrameNumber>& FrameNumbers, TArray<double>& XValues, TArray<double>& YValues, TArray<double>& ZValues)
 	{
 		int8 ChannelIndex = 0;
 
@@ -239,29 +239,29 @@ namespace DatasmithLevelSequenceImporterImpl
 			break;
 		}
 
-		// The input float values must be converted to MovieSceneFloatValue with the given interpolation
-		auto FloatConverter = [Interpolation](float InValue)
+		// The input double values must be converted to MovieSceneDoubleValue with the given interpolation
+		auto FloatConverter = [Interpolation](double InValue)
 		{
-			FMovieSceneFloatValue NewValue(InValue);
+			FMovieSceneDoubleValue NewValue(InValue);
 			NewValue.InterpMode = Interpolation;
 			return NewValue;
 		};
 
-		TArray<FMovieSceneFloatValue> MovieSceneFloatValues;
+		TArray<FMovieSceneDoubleValue> MovieSceneDoubleValues;
 
 		// X channel
-		Algo::Transform(XValues, MovieSceneFloatValues, FloatConverter);
-		Channels[ChannelIndex]->Set(FrameNumbers, MoveTemp(MovieSceneFloatValues));
+		Algo::Transform(XValues, MovieSceneDoubleValues, FloatConverter);
+		Channels[ChannelIndex]->Set(FrameNumbers, MoveTemp(MovieSceneDoubleValues));
 
 		// Y channel
-		MovieSceneFloatValues.Reset(FrameNumbers.Num());
-		Algo::Transform(YValues, MovieSceneFloatValues, FloatConverter);
-		Channels[ChannelIndex + 1]->Set(FrameNumbers, MoveTemp(MovieSceneFloatValues));
+		MovieSceneDoubleValues.Reset(FrameNumbers.Num());
+		Algo::Transform(YValues, MovieSceneDoubleValues, FloatConverter);
+		Channels[ChannelIndex + 1]->Set(FrameNumbers, MoveTemp(MovieSceneDoubleValues));
 
 		// Z channel
-		MovieSceneFloatValues.Reset(FrameNumbers.Num());
-		Algo::Transform(ZValues, MovieSceneFloatValues, FloatConverter);
-		Channels[ChannelIndex + 2]->Set(FrameNumbers, MoveTemp(MovieSceneFloatValues));
+		MovieSceneDoubleValues.Reset(FrameNumbers.Num());
+		Algo::Transform(ZValues, MovieSceneDoubleValues, FloatConverter);
+		Channels[ChannelIndex + 2]->Set(FrameNumbers, MoveTemp(MovieSceneDoubleValues));
 	}
 
 	void SetVisibilityChannels(TArrayView<FMovieSceneBoolChannel*>& Channels, ERichCurveInterpMode Interpolation, const TArray<FFrameNumber>& FrameNumbers, TArray<bool>& Values)
@@ -323,14 +323,14 @@ namespace DatasmithLevelSequenceImporterImpl
 
 		// Populate the 3D transform section with the keyframes for each transform type
 		TArray<FFrameNumber> FrameNumbers;
-		TArray<float> XValues;
-		TArray<float> YValues;
-		TArray<float> ZValues;
+		TArray<double> XValues;
+		TArray<double> YValues;
+		TArray<double> ZValues;
 
 		const FFrameRate FrameRate = MovieScene->GetDisplayRate();
 		const FFrameRate TickResolution = MovieScene->GetTickResolution();
 
-		TArrayView<FMovieSceneFloatChannel*> Channels = TransformSection->GetChannelProxy().GetChannels<FMovieSceneFloatChannel>();
+		TArrayView<FMovieSceneDoubleChannel*> Channels = TransformSection->GetChannelProxy().GetChannels<FMovieSceneDoubleChannel>();
 
 		static_assert((int)ERichCurveInterpMode::RCIM_Linear == (int)EDatasmithCurveInterpMode::Linear, "INVALID_ENUM_VALUE");
 		static_assert((int)ERichCurveInterpMode::RCIM_Constant == (int)EDatasmithCurveInterpMode::Constant, "INVALID_ENUM_VALUE");
@@ -366,7 +366,7 @@ namespace DatasmithLevelSequenceImporterImpl
 
 		FKeyDataOptimizationParams OptimParams;
 		OptimParams.DisplayRate = MovieScene->GetDisplayRate();
-		for (FMovieSceneFloatChannel* Channel : Channels)
+		for (FMovieSceneDoubleChannel* Channel : Channels)
 		{
 			Channel->Optimize(OptimParams);
 		}
diff --git a/Engine/Plugins/Experimental/MotionTrailEditorMode/Source/MotionTrailEditorMode/Private/Sequencer/MovieSceneTransformTrailTool.cpp b/Engine/Plugins/Experimental/MotionTrailEditorMode/Source/MotionTrailEditorMode/Private/Sequencer/MovieSceneTransformTrailTool.cpp
--- a/Engine/Plugins/Experimental/MotionTrailEditorMode/Source/MotionTrailEditorMode/Private/Sequencer/MovieSceneTransformTrailTool.cpp
+++ b/Engine/Plugins/Experimental/MotionTrailEditorMode/Source/MotionTrailEditorMode/Private/Sequencer/MovieSceneTransformTrailTool.cpp
@@ -8,7 +8,7 @@
 #include "ISequencer.h"
 #include "Tracks/MovieScene3DTransformTrack.h"
 #include "Sections/MovieScene3DTransformSection.h"
-#include "Channels/MovieSceneFloatChannel.h"
+#include "Channels/MovieSceneDoubleChannel.h"
 #include "Channels/MovieSceneChannelProxy.h"
 #include "EntitySystem/Interrogation/MovieSceneInterrogationLinker.h"
 #include "Systems/MovieSceneComponentTransformSystem.h"
@@ -205,14 +205,14 @@ void FDefaultMovieSceneTransformTrailTool::BuildKeys()
 	Keys.Reset();
 
 	UMovieSceneSection* AbsoluteTransformSection = OwningTrail->GetSection();
-	TArrayView<FMovieSceneFloatChannel*> FloatChannels = AbsoluteTransformSection->GetChannelProxy().GetChannels<FMovieSceneFloatChannel>();
-	FloatChannels = FloatChannels.Slice(OwningTrail->GetChannelOffset(), uint8(EMSTrailTransformChannel::MaxChannel) + 1);
+	TArrayView<FMovieSceneDoubleChannel*> DoubleChannels = AbsoluteTransformSection->GetChannelProxy().GetChannels<FMovieSceneDoubleChannel>();
+	DoubleChannels = DoubleChannels.Slice(OwningTrail->GetChannelOffset(), uint8(EMSTrailTransformChannel::MaxChannel) + 1);
 	for (int32 ChannelIdx = 0; ChannelIdx <= uint8(EMSTrailTransformChannel::MaxChannel); ChannelIdx++)
 	{
-		FMovieSceneFloatChannel* FloatChannel = FloatChannels[ChannelIdx];
-		for (int32 Idx = 0; Idx < FloatChannel->GetNumKeys(); Idx++)
+		FMovieSceneDoubleChannel* DoubleChannel = DoubleChannels[ChannelIdx];
+		for (int32 Idx = 0; Idx < DoubleChannel->GetNumKeys(); Idx++)
 		{
-			const FFrameNumber CurTime = FloatChannel->GetTimes()[Idx];
+			const FFrameNumber CurTime = DoubleChannel->GetTimes()[Idx];
 
 			if (!Keys.Contains(CurTime))
 			{
@@ -226,15 +226,15 @@ void FDefaultMovieSceneTransformTrailTool::BuildKeys()
 bool FDefaultMovieSceneTransformTrailTool::ShouldRebuildKeys()
 {
 	TMap<FFrameNumber, TSet<EMSTrailTransformChannel>> KeyTimes;
-	TArrayView<FMovieSceneFloatChannel*> FloatChannels = OwningTrail->GetSection()->GetChannelProxy().GetChannels<FMovieSceneFloatChannel>();
-	FloatChannels = FloatChannels.Slice(OwningTrail->GetChannelOffset(), uint8(EMSTrailTransformChannel::MaxChannel) + 1);
+	TArrayView<FMovieSceneDoubleChannel*> DoubleChannels = OwningTrail->GetSection()->GetChannelProxy().GetChannels<FMovieSceneDoubleChannel>();
+	DoubleChannels = DoubleChannels.Slice(OwningTrail->GetChannelOffset(), uint8(EMSTrailTransformChannel::MaxChannel) + 1);
 
 	for (uint8 ChannelIdx = 0; ChannelIdx <= uint8(EMSTrailTransformChannel::MaxChannel); ChannelIdx++)
 	{
-		FMovieSceneFloatChannel* FloatChannel = FloatChannels[ChannelIdx];
-		for (int32 Idx = 0; Idx < FloatChannel->GetNumKeys(); Idx++)
+		FMovieSceneDoubleChannel* DoubleChannel = DoubleChannels[ChannelIdx];
+		for (int32 Idx = 0; Idx < DoubleChannel->GetNumKeys(); Idx++)
 		{
-			const FFrameNumber CurTime = FloatChannel->GetTimes()[Idx];
+			const FFrameNumber CurTime = DoubleChannel->GetTimes()[Idx];
 			KeyTimes.FindOrAdd(CurTime).Add(EMSTrailTransformChannel(ChannelIdx));
 		}
 	}
@@ -254,7 +254,7 @@ bool FDefaultMovieSceneTransformTrailTool::ShouldRebuildKeys()
 		{
 			const EMSTrailTransformChannel TransformChannel = EMSTrailTransformChannel(ChannelIdx);
 			if ((!TimeKeyPair.Value.Contains(TransformChannel) && Keys[TimeKeyPair.Key]->IdxMap.Contains(TransformChannel)) ||
-				(Keys[TimeKeyPair.Key]->IdxMap.Contains(TransformChannel) && FloatChannels[uint8(TransformChannel)]->GetData().GetIndex(Keys[TimeKeyPair.Key]->IdxMap[TransformChannel]) == INDEX_NONE) ||
+				(Keys[TimeKeyPair.Key]->IdxMap.Contains(TransformChannel) && DoubleChannels[uint8(TransformChannel)]->GetData().GetIndex(Keys[TimeKeyPair.Key]->IdxMap[TransformChannel]) == INDEX_NONE) ||
 				(TimeKeyPair.Value.Contains(TransformChannel) && !Keys[TimeKeyPair.Key]->IdxMap.Contains(TransformChannel)))
 			{
 				return true;
@@ -322,7 +322,7 @@ FDefaultMovieSceneTransformTrailTool::FKeyInfo::FKeyInfo(const FFrameNumber InFr
 	, Section(InSection)
 	, OwningTrail(InOwningTrail)
 {
-	TArrayView<FMovieSceneFloatChannel*> Channels = InSection->GetChannelProxy().GetChannels<FMovieSceneFloatChannel>();
+	TArrayView<FMovieSceneDoubleChannel*> Channels = InSection->GetChannelProxy().GetChannels<FMovieSceneDoubleChannel>();
 	Channels = Channels.Slice(OwningTrail->GetChannelOffset(), uint8(EMSTrailTransformChannel::MaxChannel) + 1);
 	for (uint8 Idx = 0; Idx <= uint8(EMSTrailTransformChannel::MaxChannel); Idx++)
 	{
@@ -346,7 +346,7 @@ void  FDefaultMovieSceneTransformTrailTool::FKeyInfo::OnKeyTransformChanged(UTra
 
 void FDefaultMovieSceneTransformTrailTool::FKeyInfo::OnDragStart(UTransformProxy*)
 {
-	TArrayView<FMovieSceneFloatChannel*> Channels = Section->GetChannelProxy().GetChannels<FMovieSceneFloatChannel>();
+	TArrayView<FMovieSceneDoubleChannel*> Channels = Section->GetChannelProxy().GetChannels<FMovieSceneDoubleChannel>();
 	Channels = Channels.Slice(OwningTrail->GetChannelOffset(), uint8(EMSTrailTransformChannel::MaxChannel) + 1);
 	for (const TPair<EMSTrailTransformChannel, FKeyHandle>& ChannelHandlePair : IdxMap)
 	{
@@ -371,7 +371,7 @@ void FDefaultMovieSceneTransformTrailTool::FKeyInfo::OnDragEnd(UTransformProxy*)
 void FDefaultMovieSceneTransformTrailTool::FKeyInfo::UpdateKeyTransform(EKeyUpdateType UpdateType, FTrajectoryCache* ParentTrajectoryCache)
 {
 	bDirty = false;
-	TArrayView<FMovieSceneFloatChannel*> Channels = Section->GetChannelProxy().GetChannels<FMovieSceneFloatChannel>();
+	TArrayView<FMovieSceneDoubleChannel*> Channels = Section->GetChannelProxy().GetChannels<FMovieSceneDoubleChannel>();
 	Channels = Channels.Slice(OwningTrail->GetChannelOffset(), uint8(EMSTrailTransformChannel::MaxChannel) + 1);
 	if (UpdateType == EKeyUpdateType::FromComponentDelta)
 	{
diff --git a/Engine/Plugins/FX/Niagara/Source/Niagara/Private/MovieScene/Parameters/MovieSceneNiagaraVectorParameterTrack.cpp b/Engine/Plugins/FX/Niagara/Source/Niagara/Private/MovieScene/Parameters/MovieSceneNiagaraVectorParameterTrack.cpp
--- a/Engine/Plugins/FX/Niagara/Source/Niagara/Private/MovieScene/Parameters/MovieSceneNiagaraVectorParameterTrack.cpp
+++ b/Engine/Plugins/FX/Niagara/Source/Niagara/Private/MovieScene/Parameters/MovieSceneNiagaraVectorParameterTrack.cpp
@@ -8,19 +8,19 @@
 
 bool UMovieSceneNiagaraVectorParameterTrack::SupportsType(TSubclassOf<UMovieSceneSection> SectionClass) const
 {
-	return SectionClass == UMovieSceneVectorSection::StaticClass();
+	return SectionClass == UMovieSceneFloatVectorSection::StaticClass();
 }
 
 UMovieSceneSection* UMovieSceneNiagaraVectorParameterTrack::CreateNewSection()
 {
-	UMovieSceneVectorSection* VectorSection = NewObject<UMovieSceneVectorSection>(this, NAME_None, RF_Transactional);
+	UMovieSceneFloatVectorSection* VectorSection = NewObject<UMovieSceneFloatVectorSection>(this, NAME_None, RF_Transactional);
 	VectorSection->SetChannelsUsed(ChannelsUsed);
 	return VectorSection;
 }
 
 void UMovieSceneNiagaraVectorParameterTrack::SetSectionChannelDefaults(UMovieSceneSection* Section, const TArray<uint8>& DefaultValueData) const
 {
-	UMovieSceneVectorSection* VectorSection = Cast<UMovieSceneVectorSection>(Section);
+	UMovieSceneFloatVectorSection* VectorSection = Cast<UMovieSceneFloatVectorSection>(Section);
 	if (ensureMsgf(VectorSection != nullptr, TEXT("Section must be a color section.")) && ensureMsgf(DefaultValueData.Num() == sizeof(float) * VectorSection->GetChannelsUsed(), TEXT("DefaultValueData must be the correct vector type.")))
 	{
 		FMovieSceneChannelProxy& VectorChannelProxy = VectorSection->GetChannelProxy();
@@ -51,7 +51,7 @@ void UMovieSceneNiagaraVectorParameterTrack::SetSectionChannelDefaults(UMovieSce
 
 FMovieSceneEvalTemplatePtr UMovieSceneNiagaraVectorParameterTrack::CreateTemplateForSection(const UMovieSceneSection& InSection) const
 {
-	const UMovieSceneVectorSection* VectorSection = Cast<UMovieSceneVectorSection>(&InSection);
+	const UMovieSceneFloatVectorSection* VectorSection = Cast<UMovieSceneFloatVectorSection>(&InSection);
 	if (VectorSection != nullptr)
 	{
 		TArray<FMovieSceneFloatChannel> ComponentChannels;
@@ -72,4 +72,4 @@ int32 UMovieSceneNiagaraVectorParameterTrack::GetChannelsUsed() const
 void UMovieSceneNiagaraVectorParameterTrack::SetChannelsUsed(int32 InChannelsUsed)
 {
 	ChannelsUsed = InChannelsUsed;
-}
\ No newline at end of file
+}
diff --git a/Engine/Plugins/FX/Niagara/Source/NiagaraEditor/Private/NiagaraEditorModule.cpp b/Engine/Plugins/FX/Niagara/Source/NiagaraEditor/Private/NiagaraEditorModule.cpp
--- a/Engine/Plugins/FX/Niagara/Source/NiagaraEditor/Private/NiagaraEditorModule.cpp
+++ b/Engine/Plugins/FX/Niagara/Source/NiagaraEditor/Private/NiagaraEditorModule.cpp
@@ -997,7 +997,7 @@ void FNiagaraEditorModule::StartupModule()
 	CreateIntegerParameterTrackEditorHandle = SequencerModule.RegisterTrackEditor(FOnCreateTrackEditor::CreateStatic(
 		&FNiagaraSystemParameterTrackEditor<UMovieSceneNiagaraIntegerParameterTrack, UMovieSceneIntegerSection>::CreateTrackEditor));
 	CreateVectorParameterTrackEditorHandle = SequencerModule.RegisterTrackEditor(FOnCreateTrackEditor::CreateStatic(
-		&FNiagaraSystemParameterTrackEditor<UMovieSceneNiagaraVectorParameterTrack, UMovieSceneVectorSection>::CreateTrackEditor));
+		&FNiagaraSystemParameterTrackEditor<UMovieSceneNiagaraVectorParameterTrack, UMovieSceneFloatVectorSection>::CreateTrackEditor));
 	CreateColorParameterTrackEditorHandle = SequencerModule.RegisterTrackEditor(FOnCreateTrackEditor::CreateStatic(
 		&FNiagaraSystemColorParameterTrackEditor::CreateTrackEditor));
 
diff --git a/Engine/Plugins/MovieScene/LevelSequenceEditor/Source/LevelSequenceEditor/Private/Misc/LevelSequenceEditorActorSpawner.cpp b/Engine/Plugins/MovieScene/LevelSequenceEditor/Source/LevelSequenceEditor/Private/Misc/LevelSequenceEditorActorSpawner.cpp
--- a/Engine/Plugins/MovieScene/LevelSequenceEditor/Source/LevelSequenceEditor/Private/Misc/LevelSequenceEditorActorSpawner.cpp
+++ b/Engine/Plugins/MovieScene/LevelSequenceEditor/Source/LevelSequenceEditor/Private/Misc/LevelSequenceEditorActorSpawner.cpp
@@ -270,22 +270,22 @@ void FLevelSequenceEditorActorSpawner::SetupDefaultsForSpawnable(UObject* Spawne
 			}
 
 			// Set the section's default values to this default transform
-			TArrayView<FMovieSceneFloatChannel*> FloatChannels = TransformSection->GetChannelProxy().GetChannels<FMovieSceneFloatChannel>();
-			FloatChannels[0]->SetDefault(Location.X);
-			FloatChannels[1]->SetDefault(Location.Y);
-			FloatChannels[2]->SetDefault(Location.Z);
-
-			FloatChannels[3]->SetDefault(Rotation.X);
-			FloatChannels[4]->SetDefault(Rotation.Y);
-			FloatChannels[5]->SetDefault(Rotation.Z);
-
-			FloatChannels[6]->SetDefault(Scale.X);
-			FloatChannels[7]->SetDefault(Scale.Y);
-			FloatChannels[8]->SetDefault(Scale.Z);
+			TArrayView<FMovieSceneDoubleChannel*> DoubleChannels = TransformSection->GetChannelProxy().GetChannels<FMovieSceneDoubleChannel>();
+			DoubleChannels[0]->SetDefault(Location.X);
+			DoubleChannels[1]->SetDefault(Location.Y);
+			DoubleChannels[2]->SetDefault(Location.Z);
+
+			DoubleChannels[3]->SetDefault(Rotation.X);
+			DoubleChannels[4]->SetDefault(Rotation.Y);
+			DoubleChannels[5]->SetDefault(Rotation.Z);
+
+			DoubleChannels[6]->SetDefault(Scale.X);
+			DoubleChannels[7]->SetDefault(Scale.Y);
+			DoubleChannels[8]->SetDefault(Scale.Z);
 		}
 	}
 }
 
 #endif	// #if WITH_EDITOR
 
-#undef LOCTEXT_NAMESPACE
\ No newline at end of file
+#undef LOCTEXT_NAMESPACE
diff --git a/Engine/Plugins/MovieScene/MatineeToLevelSequence/Source/MatineeToLevelSequence/Private/CameraAnimToTemplateSequenceConverter.cpp b/Engine/Plugins/MovieScene/MatineeToLevelSequence/Source/MatineeToLevelSequence/Private/CameraAnimToTemplateSequenceConverter.cpp
--- a/Engine/Plugins/MovieScene/MatineeToLevelSequence/Source/MatineeToLevelSequence/Private/CameraAnimToTemplateSequenceConverter.cpp
+++ b/Engine/Plugins/MovieScene/MatineeToLevelSequence/Source/MatineeToLevelSequence/Private/CameraAnimToTemplateSequenceConverter.cpp
@@ -62,19 +62,19 @@ void SetupPPSettingPropertyTrack(TrackType* Track)
 }
 
 template<>
-void SetupPPSettingPropertyTrack<UMovieSceneVectorTrack, FVector2D>(UMovieSceneVectorTrack* Track)
+void SetupPPSettingPropertyTrack<UMovieSceneFloatVectorTrack, FVector2D>(UMovieSceneFloatVectorTrack* Track)
 {
 	Track->SetNumChannelsUsed(2);
 }
 
 template<>
-void SetupPPSettingPropertyTrack<UMovieSceneVectorTrack, FVector>(UMovieSceneVectorTrack* Track)
+void SetupPPSettingPropertyTrack<UMovieSceneFloatVectorTrack, FVector3f>(UMovieSceneFloatVectorTrack* Track)
 {
 	Track->SetNumChannelsUsed(3);
 }
 
 template<>
-void SetupPPSettingPropertyTrack<UMovieSceneVectorTrack, FVector4>(UMovieSceneVectorTrack* Track)
+void SetupPPSettingPropertyTrack<UMovieSceneFloatVectorTrack, FVector4>(UMovieSceneFloatVectorTrack* Track)
 {
 	Track->SetNumChannelsUsed(4);
 }
@@ -160,7 +160,7 @@ void SetDefaultPPSettingPropertyChannels<UMovieSceneFloatSection, float>(UMovieS
 
 // Vector properties.
 template<>
-void SetDefaultPPSettingPropertyChannels<UMovieSceneVectorSection, FVector2D>(UMovieSceneVectorSection* Section, const FVector2D Value)
+void SetDefaultPPSettingPropertyChannels<UMovieSceneFloatVectorSection, FVector2D>(UMovieSceneFloatVectorSection* Section, const FVector2D Value)
 {
 	Section->SetChannelsUsed(2);
 
@@ -170,7 +170,7 @@ void SetDefaultPPSettingPropertyChannels<UMovieSceneVectorSection, FVector2D>(UM
 }
 
 template<>
-void SetDefaultPPSettingPropertyChannels<UMovieSceneVectorSection, FVector>(UMovieSceneVectorSection* Section, const FVector Value)
+void SetDefaultPPSettingPropertyChannels<UMovieSceneFloatVectorSection, FVector3f>(UMovieSceneFloatVectorSection* Section, const FVector3f Value)
 {
 	Section->SetChannelsUsed(3);
 
@@ -181,7 +181,7 @@ void SetDefaultPPSettingPropertyChannels<UMovieSceneVectorSection, FVector>(UMov
 }
 
 template<>
-void SetDefaultPPSettingPropertyChannels<UMovieSceneVectorSection, FVector4>(UMovieSceneVectorSection* Section, const FVector4 Value)
+void SetDefaultPPSettingPropertyChannels<UMovieSceneFloatVectorSection, FVector4>(UMovieSceneFloatVectorSection* Section, const FVector4 Value)
 {
 	Section->SetChannelsUsed(4);
 
@@ -654,17 +654,21 @@ UObject* FCameraAnimToTemplateSequenceConverter::ConvertSingleCameraAnimToTempla
 			const FStructProperty* OverridenFieldProperty = CastField<FStructProperty>(OverridenProperty);
 			if (OverridenFieldProperty->Struct == TBaseStructure<FVector2D>::Get())
 			{
-				CreateDefaultPPSettingPropertyValueTrack<UMovieSceneVectorTrack, UMovieSceneVectorSection, FVector2D>(
+				CreateDefaultPPSettingPropertyValueTrack<UMovieSceneFloatVectorTrack, UMovieSceneFloatVectorSection, FVector2D>(
 						*NewMovieScene, CameraComponentBindingID, BasePPSettings, *OverridenFieldProperty);
 			}
-			else if (OverridenFieldProperty->Struct == TBaseStructure<FVector>::Get())
+			else if (OverridenFieldProperty->Struct->GetFName() == NAME_Vector3f
+#if UE_LARGE_WORLD_COORDINATES_DISABLED
+					|| OverridenFieldProperty->Struct->GetFName() == NAME_Vector
+#endif
+					)
 			{
-				CreateDefaultPPSettingPropertyValueTrack<UMovieSceneVectorTrack, UMovieSceneVectorSection, FVector>(
+				CreateDefaultPPSettingPropertyValueTrack<UMovieSceneFloatVectorTrack, UMovieSceneFloatVectorSection, FVector3f>(
 						*NewMovieScene, CameraComponentBindingID, BasePPSettings, *OverridenFieldProperty);
 			}
 			else if (OverridenFieldProperty->Struct == TBaseStructure<FVector4>::Get())
 			{
-				CreateDefaultPPSettingPropertyValueTrack<UMovieSceneVectorTrack, UMovieSceneVectorSection, FVector4>(
+				CreateDefaultPPSettingPropertyValueTrack<UMovieSceneFloatVectorTrack, UMovieSceneFloatVectorSection, FVector4>(
 						*NewMovieScene, CameraComponentBindingID, BasePPSettings, *OverridenFieldProperty);
 			}
 			else if (OverridenFieldProperty->Struct == TBaseStructure<FLinearColor>::Get())
diff --git a/Engine/Plugins/MovieScene/MatineeToLevelSequence/Source/MatineeToLevelSequence/Private/MatineeConverter.cpp b/Engine/Plugins/MovieScene/MatineeToLevelSequence/Source/MatineeToLevelSequence/Private/MatineeConverter.cpp
--- a/Engine/Plugins/MovieScene/MatineeToLevelSequence/Source/MatineeToLevelSequence/Private/MatineeConverter.cpp
+++ b/Engine/Plugins/MovieScene/MatineeToLevelSequence/Source/MatineeToLevelSequence/Private/MatineeConverter.cpp
@@ -352,7 +352,7 @@ void FMatineeConverter::ConvertInterpGroup(UInterpGroup* Group, FGuid ObjectBind
 			UInterpTrackVectorProp* MatineeVectorTrack = StaticCast<UInterpTrackVectorProp*>(Track);
 			if (MatineeVectorTrack->GetNumKeyframes() != 0 && ObjectBindingGuid.IsValid())
 			{
-				UMovieSceneVectorTrack* VectorTrack = AddPropertyTrack<UMovieSceneVectorTrack>(MatineeVectorTrack->PropertyName, GroupActor, ObjectBindingGuid, NewSequence, NewMovieScene, NumWarnings, BoundObjectsToGuids);
+				UMovieSceneFloatVectorTrack* VectorTrack = AddPropertyTrack<UMovieSceneFloatVectorTrack>(MatineeVectorTrack->PropertyName, GroupActor, ObjectBindingGuid, NewSequence, NewMovieScene, NumWarnings, BoundObjectsToGuids);
 				if (VectorTrack)
 				{
 					VectorTrack->SetNumChannelsUsed(3);
diff --git a/Engine/Plugins/MovieScene/SequencerScripting/Source/SequencerScripting/Private/ExtensionLibraries/MovieSceneVectorTrackExtensions.cpp b/Engine/Plugins/MovieScene/SequencerScripting/Source/SequencerScripting/Private/ExtensionLibraries/MovieSceneVectorTrackExtensions.cpp
--- a/Engine/Plugins/MovieScene/SequencerScripting/Source/SequencerScripting/Private/ExtensionLibraries/MovieSceneVectorTrackExtensions.cpp
+++ b/Engine/Plugins/MovieScene/SequencerScripting/Source/SequencerScripting/Private/ExtensionLibraries/MovieSceneVectorTrackExtensions.cpp
@@ -3,7 +3,7 @@
 #include "ExtensionLibraries/MovieSceneVectorTrackExtensions.h"
 #include "Tracks/MovieSceneVectorTrack.h"
 
-void UMovieSceneVectorTrackExtensions::SetNumChannelsUsed(UMovieSceneVectorTrack* Track, int32 InNumChannelsUsed)
+void UMovieSceneFloatVectorTrackExtensions::SetNumChannelsUsed(UMovieSceneFloatVectorTrack* Track, int32 InNumChannelsUsed)
 {
 	Track->Modify();
 
@@ -11,7 +11,21 @@ void UMovieSceneVectorTrackExtensions::SetNumChannelsUsed(UMovieSceneVectorTrack
 }
 
 
-int32 UMovieSceneVectorTrackExtensions::GetNumChannelsUsed(UMovieSceneVectorTrack* Track)
+int32 UMovieSceneFloatVectorTrackExtensions::GetNumChannelsUsed(UMovieSceneFloatVectorTrack* Track)
 {
 	return Track->GetNumChannelsUsed();
 }
+
+void UMovieSceneDoubleVectorTrackExtensions::SetNumChannelsUsed(UMovieSceneDoubleVectorTrack* Track, int32 InNumChannelsUsed)
+{
+	Track->Modify();
+
+	Track->SetNumChannelsUsed(InNumChannelsUsed);
+}
+
+
+int32 UMovieSceneDoubleVectorTrackExtensions::GetNumChannelsUsed(UMovieSceneDoubleVectorTrack* Track)
+{
+	return Track->GetNumChannelsUsed();
+}
+
diff --git a/Engine/Plugins/MovieScene/SequencerScripting/Source/SequencerScripting/Public/ExtensionLibraries/MovieSceneVectorTrackExtensions.h b/Engine/Plugins/MovieScene/SequencerScripting/Source/SequencerScripting/Public/ExtensionLibraries/MovieSceneVectorTrackExtensions.h
--- a/Engine/Plugins/MovieScene/SequencerScripting/Source/SequencerScripting/Public/ExtensionLibraries/MovieSceneVectorTrackExtensions.h
+++ b/Engine/Plugins/MovieScene/SequencerScripting/Source/SequencerScripting/Public/ExtensionLibraries/MovieSceneVectorTrackExtensions.h
@@ -8,14 +8,14 @@
 
 #include "MovieSceneVectorTrackExtensions.generated.h"
 
-
-class UMovieSceneVectorTrack;
+class UMovieSceneDoubleVectorTrack;
+class UMovieSceneFloatVectorTrack;
 
 /**
- * Function library containing methods that should be hoisted onto UMovieSceneVectorTrack for scripting
+ * Function library containing methods that should be hoisted onto UMovieSceneFloatVectorTrack for scripting
  */
 UCLASS()
-class UMovieSceneVectorTrackExtensions : public UBlueprintFunctionLibrary
+class UMovieSceneFloatVectorTrackExtensions : public UBlueprintFunctionLibrary
 {
 	GENERATED_BODY()
 
@@ -28,7 +28,7 @@ public:
 	 * @param InNumChannelsUsed The number of channels to use for this track
 	 */
 	UFUNCTION(BlueprintCallable, Category = "Sequencer|Track", meta = (ScriptMethod))
-	static void SetNumChannelsUsed(UMovieSceneVectorTrack* Track, int32 InNumChannelsUsed);
+	static void SetNumChannelsUsed(UMovieSceneFloatVectorTrack* Track, int32 InNumChannelsUsed);
 
 	/**
 	 * Get the number of channels used for this track
@@ -37,7 +37,35 @@ public:
 	 * @return The number of channels used for this track
 	 */
 	UFUNCTION(BlueprintCallable, Category = "Sequencer|Track", meta = (ScriptMethod))
-	static int32 GetNumChannelsUsed(UMovieSceneVectorTrack* Track);
+	static int32 GetNumChannelsUsed(UMovieSceneFloatVectorTrack* Track);
+};
 
+/**
+ * Function library containing methods that should be hoisted onto UMovieSceneDoubleVectorTrack for scripting
+ */
+UCLASS()
+class UMovieSceneDoubleVectorTrackExtensions : public UBlueprintFunctionLibrary
+{
+	GENERATED_BODY()
 
+public:
+
+	/**
+	 * Set the number of channels used for this track
+	 *
+	 * @param Track        The track to set
+	 * @param InNumChannelsUsed The number of channels to use for this track
+	 */
+	UFUNCTION(BlueprintCallable, Category = "Sequencer|Track", meta = (ScriptMethod))
+	static void SetNumChannelsUsed(UMovieSceneDoubleVectorTrack* Track, int32 InNumChannelsUsed);
+
+	/**
+	 * Get the number of channels used for this track
+	 *
+	 * @param Track        The track to query for the number of channels used
+	 * @return The number of channels used for this track
+	 */
+	UFUNCTION(BlueprintCallable, Category = "Sequencer|Track", meta = (ScriptMethod))
+	static int32 GetNumChannelsUsed(UMovieSceneDoubleVectorTrack* Track);
 };
+
diff --git a/Engine/Plugins/MovieScene/TemplateSequence/Source/TemplateSequence/Private/Systems/TemplateSequenceSystem.cpp b/Engine/Plugins/MovieScene/TemplateSequence/Source/TemplateSequence/Private/Systems/TemplateSequenceSystem.cpp
--- a/Engine/Plugins/MovieScene/TemplateSequence/Source/TemplateSequence/Private/Systems/TemplateSequenceSystem.cpp
+++ b/Engine/Plugins/MovieScene/TemplateSequence/Source/TemplateSequence/Private/Systems/TemplateSequenceSystem.cpp
@@ -14,8 +14,10 @@
 #include "MovieSceneTracksComponentTypes.h"
 #include "Sections/TemplateSequenceSection.h"
 #include "Systems/FloatChannelEvaluatorSystem.h"
+#include "Systems/DoubleChannelEvaluatorSystem.h"
 #include "Systems/MovieSceneBaseValueEvaluatorSystem.h"
 #include "Systems/MovieScenePiecewiseFloatBlenderSystem.h"
+#include "Systems/MovieScenePiecewiseDoubleBlenderSystem.h"
 #include "TemplateSequence.h"
 #include "TemplateSequenceComponentTypes.h"
 
@@ -260,12 +262,14 @@ UTemplateSequencePropertyScalingEvaluatorSystem::UTemplateSequencePropertyScalin
 
 	if (HasAnyFlags(RF_ClassDefaultObject))
 	{
-		// We need to wait until float channels have evaluated (both the values that we're going to multiply, and the
+		// We need to wait until float and double channels have evaluated (both the values that we're going to multiply, and the
 		// channel of the scale multiplier itself).
 		DefineImplicitPrerequisite(UFloatChannelEvaluatorSystem::StaticClass(), GetClass());
+		DefineImplicitPrerequisite(UDoubleChannelEvaluatorSystem::StaticClass(), GetClass());
 		DefineImplicitPrerequisite(UMovieSceneBaseValueEvaluatorSystem::StaticClass(), GetClass());
 		// We need to multiply values before they are blended together with other values we shouldn't touch.
 		DefineImplicitPrerequisite(GetClass(), UMovieScenePiecewiseFloatBlenderSystem::StaticClass());
+		DefineImplicitPrerequisite(GetClass(), UMovieScenePiecewiseDoubleBlenderSystem::StaticClass());
 		// We need this component to lookup the binding GUID for an entity, so we know what scaling to apply to it.
 		DefineComponentConsumer(GetClass(), BuiltInComponents->BoundObject);
 	}
@@ -361,14 +365,14 @@ struct FScaleTransformProperties
 			TRead<FInstanceHandle> InstanceHandles,
 			TRead<FMovieScenePropertyBinding> PropertyBindings,
 			TRead<FGuid> ReverseBindingLookups,
-			TReadOptional<float> BaseLocationXs, TReadOptional<float> BaseLocationYs, TReadOptional<float> BaseLocationZs,
-			TReadOptional<float> BaseRotationXs, TReadOptional<float> BaseRotationYs, TReadOptional<float> BaseRotationZs,
-			TWriteOptional<float> LocationXs, TWriteOptional<FSourceFloatChannelFlags> LocationXFlags,
-			TWriteOptional<float> LocationYs, TWriteOptional<FSourceFloatChannelFlags> LocationYFlags,
-			TWriteOptional<float> LocationZs, TWriteOptional<FSourceFloatChannelFlags> LocationZFlags,
-			TWriteOptional<float> RotationXs, TWriteOptional<FSourceFloatChannelFlags> RotationXFlags,
-			TWriteOptional<float> RotationYs, TWriteOptional<FSourceFloatChannelFlags> RotationYFlags,
-			TWriteOptional<float> RotationZs, TWriteOptional<FSourceFloatChannelFlags> RotationZFlags)
+			TReadOptional<double> BaseLocationXs, TReadOptional<double> BaseLocationYs, TReadOptional<double> BaseLocationZs,
+			TReadOptional<double> BaseRotationXs, TReadOptional<double> BaseRotationYs, TReadOptional<double> BaseRotationZs,
+			TWriteOptional<double> LocationXs, TWriteOptional<FSourceDoubleChannelFlags> LocationXFlags,
+			TWriteOptional<double> LocationYs, TWriteOptional<FSourceDoubleChannelFlags> LocationYFlags,
+			TWriteOptional<double> LocationZs, TWriteOptional<FSourceDoubleChannelFlags> LocationZFlags,
+			TWriteOptional<double> RotationXs, TWriteOptional<FSourceDoubleChannelFlags> RotationXFlags,
+			TWriteOptional<double> RotationYs, TWriteOptional<FSourceDoubleChannelFlags> RotationYFlags,
+			TWriteOptional<double> RotationZs, TWriteOptional<FSourceDoubleChannelFlags> RotationZFlags)
 	{
 		for (int32 Index = 0; Index < Allocation->Num(); ++Index)
 		{
@@ -510,24 +514,24 @@ void UTemplateSequencePropertyScalingEvaluatorSystem::OnRun(FSystemTaskPrerequis
 			.Read(BuiltInComponents->InstanceHandle)
 			.Read(BuiltInComponents->PropertyBinding)
 			.Read(TemplateSequenceComponents->PropertyScaleReverseBindingLookup)
-			.ReadOptional(BuiltInComponents->BaseFloat[0])
-			.ReadOptional(BuiltInComponents->BaseFloat[1])
-			.ReadOptional(BuiltInComponents->BaseFloat[2])
-			.ReadOptional(BuiltInComponents->BaseFloat[3])
-			.ReadOptional(BuiltInComponents->BaseFloat[4])
-			.ReadOptional(BuiltInComponents->BaseFloat[5])
-			.WriteOptional(BuiltInComponents->FloatResult[0])
-			.WriteOptional(BuiltInComponents->FloatChannelFlags[0])
-			.WriteOptional(BuiltInComponents->FloatResult[1])
-			.WriteOptional(BuiltInComponents->FloatChannelFlags[1])
-			.WriteOptional(BuiltInComponents->FloatResult[2])
-			.WriteOptional(BuiltInComponents->FloatChannelFlags[2])
-			.WriteOptional(BuiltInComponents->FloatResult[3])
-			.WriteOptional(BuiltInComponents->FloatChannelFlags[3])
-			.WriteOptional(BuiltInComponents->FloatResult[4])
-			.WriteOptional(BuiltInComponents->FloatChannelFlags[4])
-			.WriteOptional(BuiltInComponents->FloatResult[5])
-			.WriteOptional(BuiltInComponents->FloatChannelFlags[5])
+			.ReadOptional(BuiltInComponents->BaseDouble[0])
+			.ReadOptional(BuiltInComponents->BaseDouble[1])
+			.ReadOptional(BuiltInComponents->BaseDouble[2])
+			.ReadOptional(BuiltInComponents->BaseDouble[3])
+			.ReadOptional(BuiltInComponents->BaseDouble[4])
+			.ReadOptional(BuiltInComponents->BaseDouble[5])
+			.WriteOptional(BuiltInComponents->DoubleResult[0])
+			.WriteOptional(BuiltInComponents->DoubleChannelFlags[0])
+			.WriteOptional(BuiltInComponents->DoubleResult[1])
+			.WriteOptional(BuiltInComponents->DoubleChannelFlags[1])
+			.WriteOptional(BuiltInComponents->DoubleResult[2])
+			.WriteOptional(BuiltInComponents->DoubleChannelFlags[2])
+			.WriteOptional(BuiltInComponents->DoubleResult[3])
+			.WriteOptional(BuiltInComponents->DoubleChannelFlags[3])
+			.WriteOptional(BuiltInComponents->DoubleResult[4])
+			.WriteOptional(BuiltInComponents->DoubleChannelFlags[4])
+			.WriteOptional(BuiltInComponents->DoubleResult[5])
+			.WriteOptional(BuiltInComponents->DoubleChannelFlags[5])
 			.FilterAny({ TrackComponents->ComponentTransform.PropertyTag, TrackComponents->Transform.PropertyTag })
 			.Dispatch_PerAllocation<FScaleTransformProperties>(&Linker->EntityManager, InPrerequisites, &Subsequents, this);
 	}
diff --git a/Engine/Plugins/VirtualProduction/Takes/Source/TakeTrackRecorders/Private/MovieScene3DTransformTrackRecorder.cpp b/Engine/Plugins/VirtualProduction/Takes/Source/TakeTrackRecorders/Private/MovieScene3DTransformTrackRecorder.cpp
--- a/Engine/Plugins/VirtualProduction/Takes/Source/TakeTrackRecorders/Private/MovieScene3DTransformTrackRecorder.cpp
+++ b/Engine/Plugins/VirtualProduction/Takes/Source/TakeTrackRecorders/Private/MovieScene3DTransformTrackRecorder.cpp
@@ -136,16 +136,16 @@ void UMovieScene3DTransformTrackRecorder::FinalizeTrackImpl()
 					FVector Translation = DT.GetTranslation();
 					FVector EulerRotation = DT.GetRotation().Rotator().Euler();
 					FVector Scale = DT.GetScale3D();
-					TArrayView<FMovieSceneFloatChannel*> FloatChannels = MovieSceneSection->GetChannelProxy().GetChannels<FMovieSceneFloatChannel>();
-					FloatChannels[0]->SetDefault(Translation.X);
-					FloatChannels[1]->SetDefault(Translation.Y);
-					FloatChannels[2]->SetDefault(Translation.Z);
-					FloatChannels[3]->SetDefault(EulerRotation.X);
-					FloatChannels[4]->SetDefault(EulerRotation.Y);
-					FloatChannels[5]->SetDefault(EulerRotation.Z);
-					FloatChannels[6]->SetDefault(Scale.X);
-					FloatChannels[7]->SetDefault(Scale.Y);
-					FloatChannels[8]->SetDefault(Scale.Z);
+					TArrayView<FMovieSceneDoubleChannel*> DoubleChannels = MovieSceneSection->GetChannelProxy().GetChannels<FMovieSceneDoubleChannel>();
+					DoubleChannels[0]->SetDefault(Translation.X);
+					DoubleChannels[1]->SetDefault(Translation.Y);
+					DoubleChannels[2]->SetDefault(Translation.Z);
+					DoubleChannels[3]->SetDefault(EulerRotation.X);
+					DoubleChannels[4]->SetDefault(EulerRotation.Y);
+					DoubleChannels[5]->SetDefault(EulerRotation.Z);
+					DoubleChannels[6]->SetDefault(Scale.X);
+					DoubleChannels[7]->SetDefault(Scale.Y);
+					DoubleChannels[8]->SetDefault(Scale.Z);
 				}
 				if (BufferedTransforms.Times.Num() > 0)
 				{
@@ -171,52 +171,52 @@ void UMovieScene3DTransformTrackRecorder::FinalizeTrackImpl()
  	SlowTask.EnterProgressFrame();
   
  	// add buffered transforms
- 	TArrayView<FMovieSceneFloatChannel*> FloatChannels = MovieSceneSection->GetChannelProxy().GetChannels<FMovieSceneFloatChannel>();
+ 	TArrayView<FMovieSceneDoubleChannel*> DoubleChannels = MovieSceneSection->GetChannelProxy().GetChannels<FMovieSceneDoubleChannel>();
  
 	ERichCurveInterpMode LocalInterpMode = InterpolationMode;
- 	auto Transformation = [LocalInterpMode](float In)
+ 	auto Transformation = [LocalInterpMode](double In)
  	{
- 		FMovieSceneFloatValue NewValue(In);
+ 		FMovieSceneDoubleValue NewValue(In);
  		NewValue.InterpMode = LocalInterpMode;
  		return NewValue;
  	};
- 	TArray<FMovieSceneFloatValue> FloatValues;
+ 	TArray<FMovieSceneDoubleValue> DoubleValues;
  
- 	FloatValues.Reset(BufferedTransforms.LocationX.Num());
- 	Algo::Transform(BufferedTransforms.LocationX, FloatValues, Transformation);
- 	FloatChannels[0]->Set(BufferedTransforms.Times, MoveTemp(FloatValues));
+ 	DoubleValues.Reset(BufferedTransforms.LocationX.Num());
+ 	Algo::Transform(BufferedTransforms.LocationX, DoubleValues, Transformation);
+ 	DoubleChannels[0]->Set(BufferedTransforms.Times, MoveTemp(DoubleValues));
  
- 	FloatValues.Reset(BufferedTransforms.LocationY.Num());
- 	Algo::Transform(BufferedTransforms.LocationY, FloatValues, Transformation);
- 	FloatChannels[1]->Set(BufferedTransforms.Times, MoveTemp(FloatValues));
+ 	DoubleValues.Reset(BufferedTransforms.LocationY.Num());
+ 	Algo::Transform(BufferedTransforms.LocationY, DoubleValues, Transformation);
+ 	DoubleChannels[1]->Set(BufferedTransforms.Times, MoveTemp(DoubleValues));
  
- 	FloatValues.Reset(BufferedTransforms.LocationZ.Num());
- 	Algo::Transform(BufferedTransforms.LocationZ, FloatValues, Transformation);
- 	FloatChannels[2]->Set(BufferedTransforms.Times, MoveTemp(FloatValues));
+ 	DoubleValues.Reset(BufferedTransforms.LocationZ.Num());
+ 	Algo::Transform(BufferedTransforms.LocationZ, DoubleValues, Transformation);
+ 	DoubleChannels[2]->Set(BufferedTransforms.Times, MoveTemp(DoubleValues));
  
- 	FloatValues.Reset(BufferedTransforms.RotationX.Num());
- 	Algo::Transform(BufferedTransforms.RotationX, FloatValues, Transformation);
- 	FloatChannels[3]->Set(BufferedTransforms.Times, MoveTemp(FloatValues));
+ 	DoubleValues.Reset(BufferedTransforms.RotationX.Num());
+ 	Algo::Transform(BufferedTransforms.RotationX, DoubleValues, Transformation);
+ 	DoubleChannels[3]->Set(BufferedTransforms.Times, MoveTemp(DoubleValues));
  
- 	FloatValues.Reset(BufferedTransforms.RotationY.Num());
- 	Algo::Transform(BufferedTransforms.RotationY, FloatValues, Transformation);
- 	FloatChannels[4]->Set(BufferedTransforms.Times, MoveTemp(FloatValues));
+ 	DoubleValues.Reset(BufferedTransforms.RotationY.Num());
+ 	Algo::Transform(BufferedTransforms.RotationY, DoubleValues, Transformation);
+ 	DoubleChannels[4]->Set(BufferedTransforms.Times, MoveTemp(DoubleValues));
  
- 	FloatValues.Reset(BufferedTransforms.RotationZ.Num());
- 	Algo::Transform(BufferedTransforms.RotationZ, FloatValues, Transformation);
- 	FloatChannels[5]->Set(BufferedTransforms.Times, MoveTemp(FloatValues));
+ 	DoubleValues.Reset(BufferedTransforms.RotationZ.Num());
+ 	Algo::Transform(BufferedTransforms.RotationZ, DoubleValues, Transformation);
+ 	DoubleChannels[5]->Set(BufferedTransforms.Times, MoveTemp(DoubleValues));
  
- 	FloatValues.Reset(BufferedTransforms.ScaleX.Num());
- 	Algo::Transform(BufferedTransforms.ScaleX, FloatValues, Transformation);
- 	FloatChannels[6]->Set(BufferedTransforms.Times, MoveTemp(FloatValues));
+ 	DoubleValues.Reset(BufferedTransforms.ScaleX.Num());
+ 	Algo::Transform(BufferedTransforms.ScaleX, DoubleValues, Transformation);
+ 	DoubleChannels[6]->Set(BufferedTransforms.Times, MoveTemp(DoubleValues));
  
- 	FloatValues.Reset(BufferedTransforms.ScaleY.Num());
- 	Algo::Transform(BufferedTransforms.ScaleY, FloatValues, Transformation);
- 	FloatChannels[7]->Set(BufferedTransforms.Times, MoveTemp(FloatValues));
+ 	DoubleValues.Reset(BufferedTransforms.ScaleY.Num());
+ 	Algo::Transform(BufferedTransforms.ScaleY, DoubleValues, Transformation);
+ 	DoubleChannels[7]->Set(BufferedTransforms.Times, MoveTemp(DoubleValues));
  
- 	FloatValues.Reset(BufferedTransforms.ScaleZ.Num());
- 	Algo::Transform(BufferedTransforms.ScaleZ, FloatValues, Transformation);
- 	FloatChannels[8]->Set(BufferedTransforms.Times, MoveTemp(FloatValues));
+ 	DoubleValues.Reset(BufferedTransforms.ScaleZ.Num());
+ 	Algo::Transform(BufferedTransforms.ScaleZ, DoubleValues, Transformation);
+ 	DoubleChannels[8]->Set(BufferedTransforms.Times, MoveTemp(DoubleValues));
  
  	FTransform FirstTransform = FTransform::Identity;
 	if (DefaultTransform.IsSet())
@@ -241,14 +241,14 @@ void UMovieScene3DTransformTrackRecorder::FinalizeTrackImpl()
 		FKeyDataOptimizationParams Params;
 		Params.bAutoSetInterpolation = true;
 		Params.Tolerance = TrackRecorderSettings.ReduceKeysTolerance;
-		for (FMovieSceneFloatChannel* Channel : FloatChannels)
+		for (FMovieSceneDoubleChannel* Channel : DoubleChannels)
 		{
 			Channel->Optimize(Params);
 		}
 	}
 	else
 	{
-		for (FMovieSceneFloatChannel* Channel : FloatChannels)
+		for (FMovieSceneDoubleChannel* Channel : DoubleChannels)
 		{
 			Channel->AutoSetTangents();
 		}
@@ -261,7 +261,7 @@ void UMovieScene3DTransformTrackRecorder::FinalizeTrackImpl()
  		if(!bWasAttached)
  		{
  			bool bCanReset = true;
-			for (FMovieSceneFloatChannel* Channel : MovieSceneSection->GetChannelProxy().GetChannels<FMovieSceneFloatChannel>())
+			for (FMovieSceneDoubleChannel* Channel : MovieSceneSection->GetChannelProxy().GetChannels<FMovieSceneDoubleChannel>())
 			{
 				if (Channel->GetNumKeys() > 1)
 				{
@@ -272,20 +272,20 @@ void UMovieScene3DTransformTrackRecorder::FinalizeTrackImpl()
 
 			if (bCanReset)
 			{
-				for (FMovieSceneFloatChannel* Channel : MovieSceneSection->GetChannelProxy().GetChannels<FMovieSceneFloatChannel>())
+				for (FMovieSceneDoubleChannel* Channel : MovieSceneSection->GetChannelProxy().GetChannels<FMovieSceneDoubleChannel>())
 				{
 					Channel->Reset();
 				}
 
-				FloatChannels[0]->SetDefault(FirstTransform.GetTranslation()[0]);
-				FloatChannels[1]->SetDefault(FirstTransform.GetTranslation()[1]);
-				FloatChannels[2]->SetDefault(FirstTransform.GetTranslation()[2]);
-				FloatChannels[3]->SetDefault(FirstTransform.GetRotation().Euler().X);
-				FloatChannels[4]->SetDefault(FirstTransform.GetRotation().Euler().Y);
-				FloatChannels[5]->SetDefault(FirstTransform.GetRotation().Euler().Z);
-				FloatChannels[6]->SetDefault(FirstTransform.GetScale3D()[0]);
-				FloatChannels[7]->SetDefault(FirstTransform.GetScale3D()[1]);
-				FloatChannels[8]->SetDefault(FirstTransform.GetScale3D()[2]);
+				DoubleChannels[0]->SetDefault(FirstTransform.GetTranslation()[0]);
+				DoubleChannels[1]->SetDefault(FirstTransform.GetTranslation()[1]);
+				DoubleChannels[2]->SetDefault(FirstTransform.GetTranslation()[2]);
+				DoubleChannels[3]->SetDefault(FirstTransform.GetRotation().Euler().X);
+				DoubleChannels[4]->SetDefault(FirstTransform.GetRotation().Euler().Y);
+				DoubleChannels[5]->SetDefault(FirstTransform.GetRotation().Euler().Z);
+				DoubleChannels[6]->SetDefault(FirstTransform.GetScale3D()[0]);
+				DoubleChannels[7]->SetDefault(FirstTransform.GetScale3D()[1]);
+				DoubleChannels[8]->SetDefault(FirstTransform.GetScale3D()[2]);
 
 				//no longer remove spawnable transform tracks since they may not match the template
 			}
@@ -376,16 +376,16 @@ void UMovieScene3DTransformTrackRecorder::SetUpDefaultTransform()
 	
 	DefaultTransform = DT;
 	PreviousValue = DT;
-	TArrayView<FMovieSceneFloatChannel*> FloatChannels = MovieSceneSection->GetChannelProxy().GetChannels<FMovieSceneFloatChannel>();
-	FloatChannels[0]->SetDefault(Translation.X);
-	FloatChannels[1]->SetDefault(Translation.Y);
-	FloatChannels[2]->SetDefault(Translation.Z);
-	FloatChannels[3]->SetDefault(EulerRotation.X);
-	FloatChannels[4]->SetDefault(EulerRotation.Y);
-	FloatChannels[5]->SetDefault(EulerRotation.Z);
-	FloatChannels[6]->SetDefault(Scale.X);
-	FloatChannels[7]->SetDefault(Scale.Y);
-	FloatChannels[8]->SetDefault(Scale.Z);
+	TArrayView<FMovieSceneDoubleChannel*> DoubleChannels = MovieSceneSection->GetChannelProxy().GetChannels<FMovieSceneDoubleChannel>();
+	DoubleChannels[0]->SetDefault(Translation.X);
+	DoubleChannels[1]->SetDefault(Translation.Y);
+	DoubleChannels[2]->SetDefault(Translation.Z);
+	DoubleChannels[3]->SetDefault(EulerRotation.X);
+	DoubleChannels[4]->SetDefault(EulerRotation.Y);
+	DoubleChannels[5]->SetDefault(EulerRotation.Z);
+	DoubleChannels[6]->SetDefault(Scale.X);
+	DoubleChannels[7]->SetDefault(Scale.Y);
+	DoubleChannels[8]->SetDefault(Scale.Z);
 }
 
 bool UMovieScene3DTransformTrackRecorder::ResolveTransformToRecord(FTransform& OutTransform)
@@ -624,35 +624,35 @@ bool UMovieScene3DTransformTrackRecorder::LoadRecordedFile(const FString& FileNa
 							MovieSceneSection = Cast<UMovieScene3DTransformSection>(MovieSceneTrack->CreateNewSection());
 
 							MovieSceneTrack->AddSection(*MovieSceneSection);
-							TArrayView<FMovieSceneFloatChannel*> FloatChannels = MovieSceneSection->GetChannelProxy().GetChannels<FMovieSceneFloatChannel>();
+							TArrayView<FMovieSceneDoubleChannel*> DoubleChannels = MovieSceneSection->GetChannelProxy().GetChannels<FMovieSceneDoubleChannel>();
 							//MovieSceneSection->TimecodeSource = SequenceRecorderUtils::GetTimecodeSource();
 
 							bool InFirst = true;
 							// @todo: sequencer-timecode: this was previously never actually used - should it be??
 							const ERichCurveInterpMode Interpolation = RCIM_Cubic;//AnimRecorder.IsValid() ? RCIM_Linear : RCIM_Cubic;
-							TArray<FMovieSceneFloatValue> TransX;
+							TArray<FMovieSceneDoubleValue> TransX;
 							TransX.SetNumUninitialized(InFrames.Num());
-							TArray<FMovieSceneFloatValue> TransY;
+							TArray<FMovieSceneDoubleValue> TransY;
 							TransY.SetNumUninitialized(InFrames.Num());
-							TArray<FMovieSceneFloatValue> TransZ;
+							TArray<FMovieSceneDoubleValue> TransZ;
 							TransZ.SetNumUninitialized(InFrames.Num());
-							TArray<FMovieSceneFloatValue> RotX;
+							TArray<FMovieSceneDoubleValue> RotX;
 							RotX.SetNumUninitialized(InFrames.Num());
-							TArray<FMovieSceneFloatValue> RotY;
+							TArray<FMovieSceneDoubleValue> RotY;
 							RotY.SetNumUninitialized(InFrames.Num());
-							TArray<FMovieSceneFloatValue> RotZ;
+							TArray<FMovieSceneDoubleValue> RotZ;
 							RotZ.SetNumUninitialized(InFrames.Num());
-							TArray<FMovieSceneFloatValue> ScaleX;
+							TArray<FMovieSceneDoubleValue> ScaleX;
 							ScaleX.SetNumUninitialized(InFrames.Num());
-							TArray<FMovieSceneFloatValue> ScaleY;
+							TArray<FMovieSceneDoubleValue> ScaleY;
 							ScaleY.SetNumUninitialized(InFrames.Num());
-							TArray<FMovieSceneFloatValue> ScaleZ;
+							TArray<FMovieSceneDoubleValue> ScaleZ;
 							ScaleZ.SetNumUninitialized(InFrames.Num());
 							TArray<FFrameNumber> Times;
 							Times.SetNumUninitialized(InFrames.Num());
 
 							int Index = 0;
-							FMovieSceneFloatValue NewValue;
+							FMovieSceneDoubleValue NewValue;
 							NewValue.InterpMode = RCIM_Cubic; //Linear if animation ? mz
 
 							for (const FTransformSerializedFrame& SerializedFrame : InFrames)
@@ -667,15 +667,15 @@ bool UMovieScene3DTransformTrackRecorder::LoadRecordedFile(const FString& FileNa
 								{
 									MovieSceneSection->SetRange(TRange<FFrameNumber>::Inclusive(CurrentFrame, CurrentFrame));
 
-									FloatChannels[0]->SetDefault(Frame.Values[0]);
-									FloatChannels[1]->SetDefault(Frame.Values[1]);
-									FloatChannels[2]->SetDefault(Frame.Values[2]);
-									FloatChannels[3]->SetDefault(Frame.Values[3]);
-									FloatChannels[4]->SetDefault(Frame.Values[4]);
-									FloatChannels[5]->SetDefault(Frame.Values[5]);
-									FloatChannels[6]->SetDefault(Frame.Values[6]);
-									FloatChannels[7]->SetDefault(Frame.Values[7]);
-									FloatChannels[8]->SetDefault(Frame.Values[8]);
+									DoubleChannels[0]->SetDefault(Frame.Values[0]);
+									DoubleChannels[1]->SetDefault(Frame.Values[1]);
+									DoubleChannels[2]->SetDefault(Frame.Values[2]);
+									DoubleChannels[3]->SetDefault(Frame.Values[3]);
+									DoubleChannels[4]->SetDefault(Frame.Values[4]);
+									DoubleChannels[5]->SetDefault(Frame.Values[5]);
+									DoubleChannels[6]->SetDefault(Frame.Values[6]);
+									DoubleChannels[7]->SetDefault(Frame.Values[7]);
+									DoubleChannels[8]->SetDefault(Frame.Values[8]);
 									InFirst = false;
 								}
 
@@ -701,15 +701,15 @@ bool UMovieScene3DTransformTrackRecorder::LoadRecordedFile(const FString& FileNa
 								MovieSceneSection->ExpandToFrame(CurrentFrame);
 								++Index;
 							}
-							FloatChannels[0]->Set(Times, MoveTemp(TransX));
-							FloatChannels[1]->Set(Times, MoveTemp(TransY));
-							FloatChannels[2]->Set(Times, MoveTemp(TransZ));
-							FloatChannels[3]->Set(Times, MoveTemp(RotX));
-							FloatChannels[4]->Set(Times, MoveTemp(RotY));
-							FloatChannels[5]->Set(Times, MoveTemp(RotZ));
-							FloatChannels[6]->Set(Times, MoveTemp(ScaleX));
-							FloatChannels[7]->Set(Times, MoveTemp(ScaleY));
-							FloatChannels[8]->Set(Times, MoveTemp(ScaleZ));
+							DoubleChannels[0]->Set(Times, MoveTemp(TransX));
+							DoubleChannels[1]->Set(Times, MoveTemp(TransY));
+							DoubleChannels[2]->Set(Times, MoveTemp(TransZ));
+							DoubleChannels[3]->Set(Times, MoveTemp(RotX));
+							DoubleChannels[4]->Set(Times, MoveTemp(RotY));
+							DoubleChannels[5]->Set(Times, MoveTemp(RotZ));
+							DoubleChannels[6]->Set(Times, MoveTemp(ScaleX));
+							DoubleChannels[7]->Set(Times, MoveTemp(ScaleY));
+							DoubleChannels[8]->Set(Times, MoveTemp(ScaleZ));
 
 						}
 					}
@@ -728,4 +728,4 @@ bool UMovieScene3DTransformTrackRecorder::LoadRecordedFile(const FString& FileNa
 		}
 	}
 	return false;
-}
\ No newline at end of file
+}
diff --git a/Engine/Plugins/VirtualProduction/Takes/Source/TakeTrackRecorders/Private/MovieScenePropertyTrackRecorder.cpp b/Engine/Plugins/VirtualProduction/Takes/Source/TakeTrackRecorders/Private/MovieScenePropertyTrackRecorder.cpp
--- a/Engine/Plugins/VirtualProduction/Takes/Source/TakeTrackRecorders/Private/MovieScenePropertyTrackRecorder.cpp
+++ b/Engine/Plugins/VirtualProduction/Takes/Source/TakeTrackRecorders/Private/MovieScenePropertyTrackRecorder.cpp
@@ -79,8 +79,14 @@ UMovieSceneTrackRecorder* FMovieScenePropertyTrackRecorderFactory::CreateTrackRe
 	case ESerializedPropertyType::FloatType:
 		TrackRecorder->PropertyRecorder = MakeShareable(new FMovieSceneTrackPropertyRecorder<float>(Binding));
 		break;
-	case ESerializedPropertyType::VectorType:
-		TrackRecorder->PropertyRecorder = MakeShareable(new FMovieSceneTrackPropertyRecorder<FVector>(Binding));
+	case ESerializedPropertyType::DoubleType:
+		TrackRecorder->PropertyRecorder = MakeShareable(new FMovieSceneTrackPropertyRecorder<double>(Binding));
+		break;
+	case ESerializedPropertyType::Vector3fType:
+		TrackRecorder->PropertyRecorder = MakeShareable(new FMovieSceneTrackPropertyRecorder<FVector3f>(Binding));
+		break;
+	case ESerializedPropertyType::Vector3dType:
+		TrackRecorder->PropertyRecorder = MakeShareable(new FMovieSceneTrackPropertyRecorder<FVector3d>(Binding));
 		break;
 	case ESerializedPropertyType::ColorType:
 		TrackRecorder->PropertyRecorder = MakeShareable(new FMovieSceneTrackPropertyRecorder<FColor>(Binding));
@@ -122,11 +128,24 @@ void UMovieScenePropertyTrackRecorder::CreateTrackImpl()
  		}
  		else if (FStructProperty* StructProperty = CastField<FStructProperty>(Property))
  		{
- 			if (StructProperty->Struct->GetFName() == NAME_Vector)
+			// LWC_TODO: vector types
+			if (StructProperty->Struct->GetFName() == NAME_Vector3f
+#if UE_LARGE_WORLD_COORDINATES_DISABLED
+					|| StructProperty->Struct->GetFName() == NAME_Vector
+#endif
+			   )
  			{
- 				PropertyRecorder = MakeShareable(new FMovieSceneTrackPropertyRecorder<FVector>(Binding));
+ 				PropertyRecorder = MakeShareable(new FMovieSceneTrackPropertyRecorder<FVector3f>(Binding));
  			}
- 			else if (StructProperty->Struct->GetFName() == NAME_Color)
+			else if (StructProperty->Struct->GetFName() == NAME_Vector3d
+#if !UE_LARGE_WORLD_COORDINATES_DISABLED
+					|| StructProperty->Struct->GetFName() == NAME_Vector
+#endif
+			   )
+ 			{
+ 				PropertyRecorder = MakeShareable(new FMovieSceneTrackPropertyRecorder<FVector3d>(Binding));
+ 			}
+			else if (StructProperty->Struct->GetFName() == NAME_Color)
  			{
 				PropertyRecorder = MakeShareable(new FMovieSceneTrackPropertyRecorder<FColor>(Binding));
  			}
@@ -159,4 +178,4 @@ bool UMovieScenePropertyTrackRecorder::LoadRecordedFile(const FString& InFileNam
 		return PropertyRecorder->LoadRecordedFile(InFileName, InMovieScene, ActorGuidToActorMap, InCompletionCallback);
 	}
 	return false;
-}
\ No newline at end of file
+}
diff --git a/Engine/Plugins/VirtualProduction/Takes/Source/TakeTrackRecorders/Private/MovieSceneTrackPropertyRecorder.cpp b/Engine/Plugins/VirtualProduction/Takes/Source/TakeTrackRecorders/Private/MovieSceneTrackPropertyRecorder.cpp
--- a/Engine/Plugins/VirtualProduction/Takes/Source/TakeTrackRecorders/Private/MovieSceneTrackPropertyRecorder.cpp
+++ b/Engine/Plugins/VirtualProduction/Takes/Source/TakeTrackRecorders/Private/MovieSceneTrackPropertyRecorder.cpp
@@ -8,6 +8,8 @@
 #include "Tracks/MovieSceneByteTrack.h"
 #include "Sections/MovieSceneEnumSection.h"
 #include "Tracks/MovieSceneEnumTrack.h"
+#include "Sections/MovieSceneDoubleSection.h"
+#include "Tracks/MovieSceneDoubleTrack.h"
 #include "Sections/MovieSceneFloatSection.h"
 #include "Tracks/MovieSceneFloatTrack.h"
 #include "Sections/MovieSceneColorSection.h"
@@ -363,9 +365,164 @@ bool FMovieSceneTrackPropertyRecorder<uint8>::LoadRecordedFile(const FString& Fi
 }
 
 template <>
-bool FMovieSceneTrackPropertyRecorder<float>::ShouldAddNewKey(const float& InNewValue) const
+bool FMovieSceneTrackPropertyRecorder<double>::ShouldAddNewKey(const double& InNewValue) const
 {
+	return !FMath::IsNearlyEqual(PreviousValue, InNewValue);
+}
+
+template <>
+UMovieSceneSection* FMovieSceneTrackPropertyRecorder<double>::AddSection(const FString& TrackDisplayName, UMovieScene* InMovieScene, const FGuid& InGuid, bool bSetDefault)
+{
+	FName TrackName = *Binding.GetPropertyPath();
+	UMovieSceneDoubleTrack* Track = InMovieScene->FindTrack<UMovieSceneDoubleTrack>(InGuid, TrackName);
+	if (!Track)
+	{
+		Track = InMovieScene->AddTrack<UMovieSceneDoubleTrack>(InGuid);
+	}
+	else
+	{
+		Track->RemoveAllAnimationData();
+	}
+
+	if (Track)
+	{
+		Track->SetPropertyNameAndPath(Binding.GetPropertyName(), Binding.GetPropertyPath());
+
+		UMovieSceneDoubleSection* Section = Cast<UMovieSceneDoubleSection>(Track->CreateNewSection());
+
+		// We only set the track defaults when we're not loading from a serialized recording. Serialized recordings don't store channel defaults but will always store a
+		// key on the first frame which will accomplish the same.
+		if (bSetDefault)
+		{
+			FMovieSceneDoubleChannel* Channel = Section->GetChannelProxy().GetChannel<FMovieSceneDoubleChannel>(0);
+			Channel->SetDefault(PreviousValue);
+		}
+
+		Track->AddSection(*Section);
+
+		return Section;
+	}
+
+	return nullptr;
+}
 
+template <>
+void FMovieSceneTrackPropertyRecorder<double>::PostCreate(IMovieSceneTrackRecorderHost* InRecordingHost, UObject* InObjectToRecord, class UMovieScene* InMovieScene, const FGuid& InGuid, bool bOpenSerializer)
+{
+}
+
+template <>
+void FMovieSceneTrackPropertyRecorder<double>::AddKeyToSection(UMovieSceneSection* InSection, const FPropertyKey<double>& InKey)
+{
+	InSection->GetChannelProxy().GetChannel<FMovieSceneDoubleChannel>(0)->AddCubicKey(InKey.Time, InKey.Value, RCTM_Break);
+}
+
+template <>
+void FMovieSceneTrackPropertyRecorder<double>::ReduceKeys(UMovieSceneSection* InSection, float ReduceKeysTolerance)
+{
+	FKeyDataOptimizationParams Params;
+	Params.bAutoSetInterpolation = true;
+	Params.Tolerance = ReduceKeysTolerance;
+	UE::MovieScene::Optimize(InSection->GetChannelProxy().GetChannel<FMovieSceneDoubleChannel>(0), Params);
+}
+
+template <>
+double FMovieSceneTrackPropertyRecorder<double>::GetDefaultValue(UMovieSceneSection* InSection)
+{
+	FMovieSceneDoubleChannel* DoubleChannel = InSection->GetChannelProxy().GetChannels<FMovieSceneDoubleChannel>()[0];
+
+	if (DoubleChannel->GetNumKeys() > 0)
+	{
+		return DoubleChannel->GetValues()[0].Value;
+	}
+	else if (DoubleChannel->GetDefault().IsSet())
+	{
+		return DoubleChannel->GetDefault().GetValue();
+	}
+
+	return 0.f;
+}
+
+template <>
+void FMovieSceneTrackPropertyRecorder<double>::SetDefaultValue(UMovieSceneSection* InSection, const double& InDefaultValue)
+{
+	return InSection->GetChannelProxy().GetChannel<FMovieSceneDoubleChannel>(0)->SetDefault(InDefaultValue);
+}
+
+template<>
+bool FMovieSceneTrackPropertyRecorder<double>::OpenSerializer(const FString& InObjectName, const FName& InPropertyName, const FString& InTrackDisplayName, const FGuid& InGuid)
+{
+	FName SerializedType("Property");
+	FFrameRate   TickResolution = MovieSceneSection->GetTypedOuter<UMovieScene>()->GetTickResolution();
+	FPropertyFileHeader Header(TickResolution, SerializedType, InGuid);
+	Header.PropertyName = InPropertyName;
+	Header.TrackDisplayName = InTrackDisplayName;
+	Header.PropertyType = ESerializedPropertyType::DoubleType;
+	FText Error;
+	FString FileName = FString::Printf(TEXT("%s_%s_%s"), *(SerializedType.ToString()), *InObjectName, *(InPropertyName.ToString()));
+
+	if (!Serializer.OpenForWrite(FileName, Header, Error))
+	{
+		UE_LOG(PropertySerialization, Warning, TEXT("Error Opening Property File: Object '%s' Property '%s' Error: '%s'"), *InObjectName, *InPropertyName.ToString(), *Error.ToString());
+		return false;
+	}
+	return true;
+}
+
+template <>
+bool FMovieSceneTrackPropertyRecorder<double>::LoadRecordedFile(const FString& FileName, UMovieScene *InMovieScene, TMap<FGuid, AActor*>& ActorGuidToActorMap,  TFunction<void()> InCompletionCallback)
+{
+	bool bFileExists = Serializer.DoesFileExist(FileName);
+	if (bFileExists)
+	{
+		FText Error;
+		FPropertyFileHeader Header;
+
+		if (Serializer.OpenForRead(FileName, Header, Error))
+		{
+			MovieSceneSection = AddSection(Header.TrackDisplayName, InMovieScene, Header.Guid, false);
+
+			Serializer.GetDataRanges([this, InMovieScene, FileName, Header, InCompletionCallback](uint64 InMinFrameId, uint64 InMaxFrameId)
+			{
+				auto OnReadComplete = [this, InMovieScene, Header, InCompletionCallback]()
+				{
+					TArray<FPropertySerializedDoubleFrame> &InFrames = Serializer.ResultData;
+					if (InFrames.Num() > 0)
+					{
+						FFrameRate InFrameRate = Header.TickResolution;
+						for (const FPropertySerializedDoubleFrame& SerializedFrame : InFrames)
+						{
+							const FPropertySerializedDouble& Frame = SerializedFrame.Frame;
+							FFrameRate   TickResolution = MovieSceneSection->GetTypedOuter<UMovieScene>()->GetTickResolution();
+							FFrameTime FrameTime = FFrameRate::TransformTime(Frame.Time, InFrameRate, TickResolution);
+							FFrameNumber CurrentFrame = FrameTime.FrameNumber;
+							FPropertyKey<double> Key;
+							Key.Time = CurrentFrame;
+							Key.Value = Frame.Value;
+							AddKeyToSection(MovieSceneSection.Get(), Key);
+							MovieSceneSection->ExpandToFrame(CurrentFrame);
+						}
+					}
+					Serializer.Close();
+					InCompletionCallback();
+				}; //callback
+
+				Serializer.ReadFramesAtFrameRange(InMinFrameId, InMaxFrameId, OnReadComplete);
+
+			});
+			return true;
+		}
+		else
+		{
+			Serializer.Close();
+		}
+	}
+	return false;
+}
+
+template <>
+bool FMovieSceneTrackPropertyRecorder<float>::ShouldAddNewKey(const float& InNewValue) const
+{
 	return !FMath::IsNearlyEqual(PreviousValue, InNewValue);
 }
 
@@ -518,6 +675,7 @@ bool FMovieSceneTrackPropertyRecorder<float>::LoadRecordedFile(const FString& Fi
 	}
 	return false;
 }
+
 template <>
 bool FMovieSceneTrackPropertyRecorder<FColor>::ShouldAddNewKey(const FColor& InNewValue) const
 {
@@ -723,19 +881,19 @@ bool FMovieSceneTrackPropertyRecorder<FColor>::LoadRecordedFile(const FString& F
 }
 
 template <>
-bool FMovieSceneTrackPropertyRecorder<FVector>::ShouldAddNewKey(const FVector& InNewValue) const
+bool FMovieSceneTrackPropertyRecorder<FVector3f>::ShouldAddNewKey(const FVector3f& InNewValue) const
 {
 	return !FMath::IsNearlyEqual(PreviousValue.X, InNewValue.X) || !FMath::IsNearlyEqual(PreviousValue.Y, InNewValue.Y) || !FMath::IsNearlyEqual(PreviousValue.Z, InNewValue.Z);
 }
 
 template <>
-UMovieSceneSection* FMovieSceneTrackPropertyRecorder<FVector>::AddSection(const FString& TrackDisplayName, UMovieScene* InMovieScene, const FGuid& InGuid, bool bSetDefault)
+UMovieSceneSection* FMovieSceneTrackPropertyRecorder<FVector3f>::AddSection(const FString& TrackDisplayName, UMovieScene* InMovieScene, const FGuid& InGuid, bool bSetDefault)
 {
 	FName TrackName = *Binding.GetPropertyPath();
-	UMovieSceneVectorTrack* Track = InMovieScene->FindTrack<UMovieSceneVectorTrack>(InGuid, TrackName);
+	UMovieSceneFloatVectorTrack* Track = InMovieScene->FindTrack<UMovieSceneFloatVectorTrack>(InGuid, TrackName);
 	if (!Track)
 	{
-		Track = InMovieScene->AddTrack<UMovieSceneVectorTrack>(InGuid);
+		Track = InMovieScene->AddTrack<UMovieSceneFloatVectorTrack>(InGuid);
 	}
 	else
 	{
@@ -748,7 +906,7 @@ UMovieSceneSection* FMovieSceneTrackPropertyRecorder<FVector>::AddSection(const
 		Track->SetNumChannelsUsed(3);
 		Track->SetPropertyNameAndPath(Binding.GetPropertyName(), Binding.GetPropertyPath());
 
-		UMovieSceneVectorSection* Section = Cast<UMovieSceneVectorSection>(Track->CreateNewSection());
+		UMovieSceneFloatVectorSection* Section = Cast<UMovieSceneFloatVectorSection>(Track->CreateNewSection());
 
 		// We only set the track defaults when we're not loading from a serialized recording. Serialized recordings don't store channel defaults but will always store a
 		// key on the first frame which will accomplish the same.
@@ -768,12 +926,12 @@ UMovieSceneSection* FMovieSceneTrackPropertyRecorder<FVector>::AddSection(const
 }
 
 template <>
-void FMovieSceneTrackPropertyRecorder<FVector>::PostCreate(IMovieSceneTrackRecorderHost* InRecordingHost, UObject* InObjectToRecord, class UMovieScene* InMovieScene, const FGuid& InGuid, bool bOpenSerializer)
+void FMovieSceneTrackPropertyRecorder<FVector3f>::PostCreate(IMovieSceneTrackRecorderHost* InRecordingHost, UObject* InObjectToRecord, class UMovieScene* InMovieScene, const FGuid& InGuid, bool bOpenSerializer)
 {
 }
 
 template <>
-void FMovieSceneTrackPropertyRecorder<FVector>::AddKeyToSection(UMovieSceneSection* InSection, const FPropertyKey<FVector>& InKey)
+void FMovieSceneTrackPropertyRecorder<FVector3f>::AddKeyToSection(UMovieSceneSection* InSection, const FPropertyKey<FVector3f>& InKey)
 {
 	TArrayView<FMovieSceneFloatChannel*> FloatChannels = InSection->GetChannelProxy().GetChannels<FMovieSceneFloatChannel>();
 	FloatChannels[0]->AddCubicKey(InKey.Time, InKey.Value.X, RCTM_Break);
@@ -782,7 +940,7 @@ void FMovieSceneTrackPropertyRecorder<FVector>::AddKeyToSection(UMovieSceneSecti
 }
 
 template <>
-void FMovieSceneTrackPropertyRecorder<FVector>::ReduceKeys(UMovieSceneSection* InSection, float ReduceKeysTolerance)
+void FMovieSceneTrackPropertyRecorder<FVector3f>::ReduceKeys(UMovieSceneSection* InSection, float ReduceKeysTolerance)
 {
 	TArrayView<FMovieSceneFloatChannel*> FloatChannels = InSection->GetChannelProxy().GetChannels<FMovieSceneFloatChannel>();
 
@@ -795,9 +953,9 @@ void FMovieSceneTrackPropertyRecorder<FVector>::ReduceKeys(UMovieSceneSection* I
 }
 
 template <>
-FVector FMovieSceneTrackPropertyRecorder<FVector>::GetDefaultValue(UMovieSceneSection* InSection)
+FVector3f FMovieSceneTrackPropertyRecorder<FVector3f>::GetDefaultValue(UMovieSceneSection* InSection)
 {
-	FVector DefaultValue(0.f);
+	FVector3f DefaultValue(0.f);
 
 	TArrayView<FMovieSceneFloatChannel*> FloatChannels = InSection->GetChannelProxy().GetChannels<FMovieSceneFloatChannel>();
 
@@ -817,7 +975,7 @@ FVector FMovieSceneTrackPropertyRecorder<FVector>::GetDefaultValue(UMovieSceneSe
 }
 
 template <>
-void FMovieSceneTrackPropertyRecorder<FVector>::SetDefaultValue(UMovieSceneSection* InSection, const FVector& InDefaultValue)
+void FMovieSceneTrackPropertyRecorder<FVector3f>::SetDefaultValue(UMovieSceneSection* InSection, const FVector3f& InDefaultValue)
 {
 	TArrayView<FMovieSceneFloatChannel*> FloatChannels = InSection->GetChannelProxy().GetChannels<FMovieSceneFloatChannel>();
 	FloatChannels[0]->SetDefault(InDefaultValue[0]);
@@ -826,14 +984,14 @@ void FMovieSceneTrackPropertyRecorder<FVector>::SetDefaultValue(UMovieSceneSecti
 }
 
 template<>
-bool FMovieSceneTrackPropertyRecorder<FVector>::OpenSerializer(const FString& InObjectName, const FName& InPropertyName, const FString& InTrackDisplayName, const FGuid& InGuid)
+bool FMovieSceneTrackPropertyRecorder<FVector3f>::OpenSerializer(const FString& InObjectName, const FName& InPropertyName, const FString& InTrackDisplayName, const FGuid& InGuid)
 {
 	FName SerializedType("Property");
 	FFrameRate   TickResolution = MovieSceneSection->GetTypedOuter<UMovieScene>()->GetTickResolution();
 	FPropertyFileHeader Header(TickResolution, SerializedType, InGuid);
 	Header.PropertyName = InPropertyName;
 	Header.TrackDisplayName = InTrackDisplayName;
-	Header.PropertyType = ESerializedPropertyType::VectorType;
+	Header.PropertyType = ESerializedPropertyType::Vector3fType;
 	
 	FText Error;
 	FString FileName = FString::Printf(TEXT("%s_%s_%s"), *(SerializedType.ToString()), *InObjectName, *(InPropertyName.ToString()));
@@ -847,7 +1005,7 @@ bool FMovieSceneTrackPropertyRecorder<FVector>::OpenSerializer(const FString& In
 }
 
 template <>
-bool FMovieSceneTrackPropertyRecorder<FVector>::LoadRecordedFile(const FString& FileName, UMovieScene *InMovieScene, TMap<FGuid, AActor*>& ActorGuidToActorMap, TFunction<void()> InCompletionCallback)
+bool FMovieSceneTrackPropertyRecorder<FVector3f>::LoadRecordedFile(const FString& FileName, UMovieScene *InMovieScene, TMap<FGuid, AActor*>& ActorGuidToActorMap, TFunction<void()> InCompletionCallback)
 {
 	bool bFileExists = Serializer.DoesFileExist(FileName);
 	if (bFileExists)
@@ -863,17 +1021,17 @@ bool FMovieSceneTrackPropertyRecorder<FVector>::LoadRecordedFile(const FString&
 			{
 				auto OnReadComplete = [this, InMovieScene, Header, InCompletionCallback]()
 				{
-					TArray<FPropertySerializedVectorFrame> &InFrames = Serializer.ResultData;
+					TArray<FPropertySerializedVector3fFrame> &InFrames = Serializer.ResultData;
 					if (InFrames.Num() > 0)
 					{
 						FFrameRate InFrameRate = Header.TickResolution;
-						for (const FPropertySerializedVectorFrame& SerializedFrame : InFrames)
+						for (const FPropertySerializedVector3fFrame& SerializedFrame : InFrames)
 						{
-							const FPropertySerializedVector& Frame = SerializedFrame.Frame;
+							const FPropertySerializedVector3f& Frame = SerializedFrame.Frame;
 							FFrameRate   TickResolution = MovieSceneSection->GetTypedOuter<UMovieScene>()->GetTickResolution();
 							FFrameTime FrameTime = FFrameRate::TransformTime(Frame.Time, InFrameRate, TickResolution);
 							FFrameNumber CurrentFrame = FrameTime.FrameNumber;
-							FPropertyKey<FVector> Key;
+							FPropertyKey<FVector3f> Key;
 							Key.Time = CurrentFrame;
 							Key.Value = Frame.Value;
 							AddKeyToSection(MovieSceneSection.Get(), Key);
@@ -897,7 +1055,180 @@ bool FMovieSceneTrackPropertyRecorder<FVector>::LoadRecordedFile(const FString&
 	return false;
 }
 
+template <>
+bool FMovieSceneTrackPropertyRecorder<FVector3d>::ShouldAddNewKey(const FVector3d& InNewValue) const
+{
+	return !FMath::IsNearlyEqual(PreviousValue.X, InNewValue.X) || !FMath::IsNearlyEqual(PreviousValue.Y, InNewValue.Y) || !FMath::IsNearlyEqual(PreviousValue.Z, InNewValue.Z);
+}
+
+template <>
+UMovieSceneSection* FMovieSceneTrackPropertyRecorder<FVector3d>::AddSection(const FString& TrackDisplayName, UMovieScene* InMovieScene, const FGuid& InGuid, bool bSetDefault)
+{
+	FName TrackName = *Binding.GetPropertyPath();
+	UMovieSceneDoubleVectorTrack* Track = InMovieScene->FindTrack<UMovieSceneDoubleVectorTrack>(InGuid, TrackName);
+	if (!Track)
+	{
+		Track = InMovieScene->AddTrack<UMovieSceneDoubleVectorTrack>(InGuid);
+	}
+	else
+	{
+		Track->RemoveAllAnimationData();
+	}
+
+
+	if (Track)
+	{
+		Track->SetNumChannelsUsed(3);
+		Track->SetPropertyNameAndPath(Binding.GetPropertyName(), Binding.GetPropertyPath());
+
+		UMovieSceneDoubleVectorSection* Section = Cast<UMovieSceneDoubleVectorSection>(Track->CreateNewSection());
+
+		// We only set the track defaults when we're not loading from a serialized recording. Serialized recordings don't store channel defaults but will always store a
+		// key on the first frame which will accomplish the same.
+		if (bSetDefault)
+		{
+			TArrayView<FMovieSceneDoubleChannel*> DoubleChannels = Section->GetChannelProxy().GetChannels<FMovieSceneDoubleChannel>();
+			DoubleChannels[0]->SetDefault(PreviousValue.X);
+			DoubleChannels[1]->SetDefault(PreviousValue.Y);
+			DoubleChannels[2]->SetDefault(PreviousValue.Z);
+		}
+
+		Track->AddSection(*Section);
+		return Section;
+	}
+
+	return nullptr;
+}
 
+template <>
+void FMovieSceneTrackPropertyRecorder<FVector3d>::PostCreate(IMovieSceneTrackRecorderHost* InRecordingHost, UObject* InObjectToRecord, class UMovieScene* InMovieScene, const FGuid& InGuid, bool bOpenSerializer)
+{
+}
+
+template <>
+void FMovieSceneTrackPropertyRecorder<FVector3d>::AddKeyToSection(UMovieSceneSection* InSection, const FPropertyKey<FVector3d>& InKey)
+{
+	TArrayView<FMovieSceneDoubleChannel*> DoubleChannels = InSection->GetChannelProxy().GetChannels<FMovieSceneDoubleChannel>();
+	DoubleChannels[0]->AddCubicKey(InKey.Time, InKey.Value.X, RCTM_Break);
+	DoubleChannels[1]->AddCubicKey(InKey.Time, InKey.Value.Y, RCTM_Break);
+	DoubleChannels[2]->AddCubicKey(InKey.Time, InKey.Value.Z, RCTM_Break);
+}
+
+template <>
+void FMovieSceneTrackPropertyRecorder<FVector3d>::ReduceKeys(UMovieSceneSection* InSection, float ReduceKeysTolerance)
+{
+	TArrayView<FMovieSceneDoubleChannel*> DoubleChannels = InSection->GetChannelProxy().GetChannels<FMovieSceneDoubleChannel>();
+
+	FKeyDataOptimizationParams Params;
+	Params.bAutoSetInterpolation = true;
+	Params.Tolerance = ReduceKeysTolerance;
+	UE::MovieScene::Optimize(DoubleChannels[0], Params);
+	UE::MovieScene::Optimize(DoubleChannels[1], Params);
+	UE::MovieScene::Optimize(DoubleChannels[2], Params);
+}
+
+template <>
+FVector3d FMovieSceneTrackPropertyRecorder<FVector3d>::GetDefaultValue(UMovieSceneSection* InSection)
+{
+	FVector3d DefaultValue(0.f);
+
+	TArrayView<FMovieSceneDoubleChannel*> DoubleChannels = InSection->GetChannelProxy().GetChannels<FMovieSceneDoubleChannel>();
+
+	for (int32 ChannelIndex = 0; ChannelIndex < 3; ++ChannelIndex)
+	{
+		if (DoubleChannels[ChannelIndex]->GetNumKeys() > 0)
+		{
+			DefaultValue[ChannelIndex] = DoubleChannels[ChannelIndex]->GetValues()[0].Value;
+		}
+		else if (DoubleChannels[ChannelIndex]->GetDefault().IsSet())
+		{
+			DefaultValue[ChannelIndex] = DoubleChannels[ChannelIndex]->GetDefault().GetValue();
+		}
+	}
+
+	return DefaultValue;
+}
+
+template <>
+void FMovieSceneTrackPropertyRecorder<FVector3d>::SetDefaultValue(UMovieSceneSection* InSection, const FVector3d& InDefaultValue)
+{
+	TArrayView<FMovieSceneDoubleChannel*> DoubleChannels = InSection->GetChannelProxy().GetChannels<FMovieSceneDoubleChannel>();
+	DoubleChannels[0]->SetDefault(InDefaultValue[0]);
+	DoubleChannels[1]->SetDefault(InDefaultValue[1]);
+	DoubleChannels[2]->SetDefault(InDefaultValue[2]);
+}
+
+template<>
+bool FMovieSceneTrackPropertyRecorder<FVector3d>::OpenSerializer(const FString& InObjectName, const FName& InPropertyName, const FString& InTrackDisplayName, const FGuid& InGuid)
+{
+	FName SerializedType("Property");
+	FFrameRate   TickResolution = MovieSceneSection->GetTypedOuter<UMovieScene>()->GetTickResolution();
+	FPropertyFileHeader Header(TickResolution, SerializedType, InGuid);
+	Header.PropertyName = InPropertyName;
+	Header.TrackDisplayName = InTrackDisplayName;
+	Header.PropertyType = ESerializedPropertyType::Vector3dType;
+	
+	FText Error;
+	FString FileName = FString::Printf(TEXT("%s_%s_%s"), *(SerializedType.ToString()), *InObjectName, *(InPropertyName.ToString()));
+
+	if (!Serializer.OpenForWrite(FileName, Header, Error))
+	{
+		UE_LOG(PropertySerialization, Warning, TEXT("Error Opening Property File: Object: '%s' Property '%s' Error: '%s'"), *InObjectName, *InPropertyName.ToString(), *Error.ToString());
+		return false;
+	}
+	return true;
+}
+
+template <>
+bool FMovieSceneTrackPropertyRecorder<FVector3d>::LoadRecordedFile(const FString& FileName, UMovieScene *InMovieScene, TMap<FGuid, AActor*>& ActorGuidToActorMap, TFunction<void()> InCompletionCallback)
+{
+	bool bFileExists = Serializer.DoesFileExist(FileName);
+	if (bFileExists)
+	{
+		FText Error;
+		FPropertyFileHeader Header;
+
+		if (Serializer.OpenForRead(FileName, Header, Error))
+		{
+			MovieSceneSection = AddSection(Header.TrackDisplayName, InMovieScene, Header.Guid, false);
+
+			Serializer.GetDataRanges([this, InMovieScene, FileName, Header, InCompletionCallback](uint64 InMinFrameId, uint64 InMaxFrameId)
+			{
+				auto OnReadComplete = [this, InMovieScene, Header, InCompletionCallback]()
+				{
+					TArray<FPropertySerializedVector3dFrame> &InFrames = Serializer.ResultData;
+					if (InFrames.Num() > 0)
+					{
+						FFrameRate InFrameRate = Header.TickResolution;
+						for (const FPropertySerializedVector3dFrame& SerializedFrame : InFrames)
+						{
+							const FPropertySerializedVector3d& Frame = SerializedFrame.Frame;
+							FFrameRate   TickResolution = MovieSceneSection->GetTypedOuter<UMovieScene>()->GetTickResolution();
+							FFrameTime FrameTime = FFrameRate::TransformTime(Frame.Time, InFrameRate, TickResolution);
+							FFrameNumber CurrentFrame = FrameTime.FrameNumber;
+							FPropertyKey<FVector3d> Key;
+							Key.Time = CurrentFrame;
+							Key.Value = Frame.Value;
+							AddKeyToSection(MovieSceneSection.Get(), Key);
+							MovieSceneSection->ExpandToFrame(CurrentFrame);
+						}
+					}
+					Serializer.Close();
+					InCompletionCallback();
+				}; //callback
+
+				Serializer.ReadFramesAtFrameRange(InMinFrameId, InMaxFrameId, OnReadComplete);
+
+			});
+			return true;
+		}
+		else
+		{
+			Serializer.Close();
+		}
+	}
+	return false;
+}
 
 template <>
 bool FMovieSceneTrackPropertyRecorder<int32>::ShouldAddNewKey(const int32& InNewValue) const
diff --git a/Engine/Plugins/VirtualProduction/Takes/Source/TakeTrackRecorders/Public/TrackRecorders/MovieScene3DTransformTrackRecorder.h b/Engine/Plugins/VirtualProduction/Takes/Source/TakeTrackRecorders/Public/TrackRecorders/MovieScene3DTransformTrackRecorder.h
--- a/Engine/Plugins/VirtualProduction/Takes/Source/TakeTrackRecorders/Public/TrackRecorders/MovieScene3DTransformTrackRecorder.h
+++ b/Engine/Plugins/VirtualProduction/Takes/Source/TakeTrackRecorders/Public/TrackRecorders/MovieScene3DTransformTrackRecorder.h
@@ -44,9 +44,9 @@ struct FBufferedTransformKeys
 {
 
 	TArray<FFrameNumber> Times;
-	TArray<float> LocationX, LocationY, LocationZ;
-	TArray<float> RotationX, RotationY, RotationZ;
-	TArray<float> ScaleX, ScaleY, ScaleZ;
+	TArray<double> LocationX, LocationY, LocationZ;
+	TArray<double> RotationX, RotationY, RotationZ;
+	TArray<double> ScaleX, ScaleY, ScaleZ;
 	//Currently used only when needed when collapsing when ProcessingAnimation()
 	TArray<FQuat> CachedQuats;
 
