Unshelved from pending changelist '17523503':

Take Recoder: Add Remove Root Option to animation recording so that we can specify to leave root animation on the anim sequence when recording, (was an option on old Sequence Recorder).

diff --git a/Engine/Plugins/VirtualProduction/Takes/Source/TakeRecorderSources/Private/TakeRecorderActorSource.cpp b/Engine/Plugins/VirtualProduction/Takes/Source/TakeRecorderSources/Private/TakeRecorderActorSource.cpp
--- a/Engine/Plugins/VirtualProduction/Takes/Source/TakeRecorderSources/Private/TakeRecorderActorSource.cpp
+++ b/Engine/Plugins/VirtualProduction/Takes/Source/TakeRecorderSources/Private/TakeRecorderActorSource.cpp
@@ -860,10 +860,9 @@ void UTakeRecorderActorSource::PostProcessTrackRecorders(ULevelSequence* InSeque
 			break;
 		}
 	}
-
-	if (RootTransformRecorder && FirstAnimationRecorder)
+	// We need to take the root motion data from the animation and override the data the Transform Track had originally captured if we are removing root
+	if (RootTransformRecorder && FirstAnimationRecorder && FirstAnimationRecorder->RootWasRemoved())
 	{
-		// We need to take the root motion data from the animation and override the data the Transform Track had originally captured
 		RootTransformRecorder->PostProcessAnimationData(FirstAnimationRecorder);
 		FirstAnimationRecorder->RemoveRootMotion();
 	}
@@ -875,7 +874,7 @@ void UTakeRecorderActorSource::PostProcessTrackRecorders(ULevelSequence* InSeque
 		{
 			UMovieSceneAnimationTrackRecorder* AnimationTrackRecorder = Cast<UMovieSceneAnimationTrackRecorder>(TrackRecorder);
 			
-			if (TrackRecorder != FirstAnimationRecorder)
+			if (TrackRecorder != FirstAnimationRecorder && AnimationTrackRecorder->RootWasRemoved())
 			{
 				AnimationTrackRecorder->RemoveRootMotion();
 			}
diff --git a/Engine/Plugins/VirtualProduction/Takes/Source/TakeTrackRecorders/Private/MovieScene3DTransformTrackRecorder.cpp b/Engine/Plugins/VirtualProduction/Takes/Source/TakeTrackRecorders/Private/MovieScene3DTransformTrackRecorder.cpp
--- a/Engine/Plugins/VirtualProduction/Takes/Source/TakeTrackRecorders/Private/MovieScene3DTransformTrackRecorder.cpp
+++ b/Engine/Plugins/VirtualProduction/Takes/Source/TakeTrackRecorders/Private/MovieScene3DTransformTrackRecorder.cpp
@@ -463,8 +463,7 @@ void UMovieScene3DTransformTrackRecorder::PostProcessAnimationData(UMovieSceneAn
 	{
 		return;
 	}
-	UMovieSceneAnimationTrackRecorderSettings* AnimSettings = CastChecked<UMovieSceneAnimationTrackRecorderSettings>(AnimTrackRecorder->GetTrackRecorderSettings());
-	if (AnimSettings->bRemoveRootAnimation)
+	if (AnimTrackRecorder->RootWasRemoved())
 	{
 		//Get All Animation Keys
 		FBufferedTransformKeys  AnimationKeys;
diff --git a/Engine/Plugins/VirtualProduction/Takes/Source/TakeTrackRecorders/Private/MovieSceneAnimationTrackRecorder.cpp b/Engine/Plugins/VirtualProduction/Takes/Source/TakeTrackRecorders/Private/MovieSceneAnimationTrackRecorder.cpp
--- a/Engine/Plugins/VirtualProduction/Takes/Source/TakeTrackRecorders/Private/MovieSceneAnimationTrackRecorder.cpp
+++ b/Engine/Plugins/VirtualProduction/Takes/Source/TakeTrackRecorders/Private/MovieSceneAnimationTrackRecorder.cpp
@@ -247,14 +247,14 @@ void UMovieSceneAnimationTrackRecorder::RecordSampleImpl(const FQualifiedFrameTi
 		USceneComponent* RootComponent = Actor->GetRootComponent();
 		USceneComponent* AttachParent = RootComponent ? RootComponent->GetAttachParent() : nullptr;
 
-		//In Sequence Recorder this would be done via checking if the component was dynamically created, due to changes in how the take recorder handles this, it no longer 
-		//possible so it seems if it's native do root, otherwise, don't.  
-		bool bRemoveRootAnimation = SkeletalMeshComponent->CreationMethod != EComponentCreationMethod::Native ? false : true;
-		// Need to pass this up to the settings since it's used later to force root lock and transfer root from animation to transform
 		UMovieSceneAnimationTrackRecorderSettings* AnimSettings = CastChecked<UMovieSceneAnimationTrackRecorderSettings>(Settings.Get());
-		AnimSettings->bRemoveRootAnimation = bRemoveRootAnimation;
+		//In Sequence Recorder this would be done via checking if the component was dynamically created, due to changes in how the take recorder handles this, it no longer 
+		//possible so it seems if it's native do root, otherwise use the setting
+		//we what we did on the track recorder since later we need to actually remove the root and transfer to the transform track if needed.
+		bRootWasRemoved = SkeletalMeshComponent->CreationMethod != EComponentCreationMethod::Native ? false : AnimSettings->bRemoveRootAnimation;
+
 		//If not removing root we also don't record in world space ( not totally sure if it matters but matching up with Sequence Recorder)
-		bool bRecordInWorldSpace = bRemoveRootAnimation == false ? false : true;
+		bool bRecordInWorldSpace = bRootWasRemoved == false ? false : true;
 
 		if (bRecordInWorldSpace && AttachParent && OwningTakeRecorderSource)
 		{
@@ -271,7 +271,7 @@ void UMovieSceneAnimationTrackRecorder::RecordSampleImpl(const FQualifiedFrameTi
 		RecordingSettings.TangentMode = AnimSettings->TangentMode;
 		RecordingSettings.Length = 0;
 		RecordingSettings.bRecordInWorldSpace = bRecordInWorldSpace;
-		RecordingSettings.bRemoveRootAnimation = bRemoveRootAnimation;
+		RecordingSettings.bRemoveRootAnimation = bRootWasRemoved;
 		RecordingSettings.bCheckDeltaTimeAtBeginning = false;
 		AnimationRecorder.Init(SkeletalMeshComponent.Get(), AnimSequence.Get(), &AnimationSerializer, RecordingSettings);
 		AnimationRecorder.BeginRecording();
@@ -297,8 +297,7 @@ void UMovieSceneAnimationTrackRecorder::RemoveRootMotion()
 {
 	 if(AnimSequence.IsValid())
 	 {
-		 UMovieSceneAnimationTrackRecorderSettings* AnimSettings = CastChecked<UMovieSceneAnimationTrackRecorderSettings>(Settings.Get());
-		 if (AnimSettings->bRemoveRootAnimation) 
+		 if (bRootWasRemoved)
 		 {
 			 // Remove Root Motion by forcing the root lock on for now (which prevents the motion at evaluation time)
 			 // In addition to set it to root lock we need to make sure it's to be zero'd since 
diff --git a/Engine/Plugins/VirtualProduction/Takes/Source/TakeTrackRecorders/Public/TrackRecorders/MovieSceneAnimationTrackRecorder.h b/Engine/Plugins/VirtualProduction/Takes/Source/TakeTrackRecorders/Public/TrackRecorders/MovieSceneAnimationTrackRecorder.h
--- a/Engine/Plugins/VirtualProduction/Takes/Source/TakeTrackRecorders/Public/TrackRecorders/MovieSceneAnimationTrackRecorder.h
+++ b/Engine/Plugins/VirtualProduction/Takes/Source/TakeTrackRecorders/Public/TrackRecorders/MovieSceneAnimationTrackRecorder.h
@@ -56,6 +56,7 @@ protected:
 	// UMovieSceneTrackRecorder Interface
 
 public:
+	bool RootWasRemoved() const { return bRootWasRemoved; }
 	void RemoveRootMotion();
 	void ProcessRecordedTimes(const FString& HoursName, const FString& MinutesName, const FString& SecondsName, const FString& FramesName, const FString& SubFramesName, const FString& SlateName, const FString& Slate);
 
@@ -94,4 +95,7 @@ private:
 
 	/**Serializer */
 	FAnimationSerializer AnimationSerializer;
+
+	/** Root Was Removed*/
+	bool bRootWasRemoved = true;
 };
diff --git a/Engine/Plugins/VirtualProduction/Takes/Source/TakeTrackRecorders/Public/TrackRecorders/MovieSceneAnimationTrackRecorderSettings.h b/Engine/Plugins/VirtualProduction/Takes/Source/TakeTrackRecorders/Public/TrackRecorders/MovieSceneAnimationTrackRecorderSettings.h
--- a/Engine/Plugins/VirtualProduction/Takes/Source/TakeTrackRecorders/Public/TrackRecorders/MovieSceneAnimationTrackRecorderSettings.h
+++ b/Engine/Plugins/VirtualProduction/Takes/Source/TakeTrackRecorders/Public/TrackRecorders/MovieSceneAnimationTrackRecorderSettings.h
@@ -44,8 +44,8 @@ public:
 	UPROPERTY(EditAnywhere, Category = "Animation Recorder Settings")
 	TEnumAsByte<ERichCurveTangentMode> TangentMode;
 
-	/** The following parameter is dynamically set based upon whether or not the animation was spawned dynamically via a blueprint or not, if so set to false, otherwise true */
-	UPROPERTY()
+	/** If true we remove the root animation and move it to a transform track, if false we leave it on the root bone in the anim sequence*/
+	UPROPERTY(config, EditAnywhere, BlueprintReadWrite, Category = "Animation Recorder Settings")
 	bool bRemoveRootAnimation;
 
 
