TakeRecorder: Added a setting and command line option to suppress the recording notification widget.

A main motivation for this option is that notification windows interfere with independent flip in fullscreen mode when recording in an nDisplay render node.
When whe notification window is disabled, it will Log start/stop.
FApp::Unattended will also force the disabling of the notification window.

diff --git a/Engine/Plugins/VirtualProduction/Takes/Source/TakeRecorder/Private/Recorder/TakeRecorder.cpp b/Engine/Plugins/VirtualProduction/Takes/Source/TakeRecorder/Private/Recorder/TakeRecorder.cpp
--- a/Engine/Plugins/VirtualProduction/Takes/Source/TakeRecorder/Private/Recorder/TakeRecorder.cpp
+++ b/Engine/Plugins/VirtualProduction/Takes/Source/TakeRecorder/Private/Recorder/TakeRecorder.cpp
@@ -1,22 +1,26 @@
 // Copyright Epic Games, Inc. All Rights Reserved.
 
 #include "Recorder/TakeRecorder.h"
+
+#include "AssetRegistryModule.h"
+#include "CoreGlobals.h"
 #include "Engine/Engine.h"
 #include "Engine/EngineTypes.h"
+#include "IAssetRegistry.h"
+#include "ISequencer.h"
+#include "LevelSequence.h"
+#include "Misc/CommandLine.h"
+#include "Misc/Parse.h"
+#include "MovieSceneTimeHelpers.h"
 #include "Recorder/TakeRecorderBlueprintLibrary.h"
-#include "TakePreset.h"
 #include "TakeMetaData.h"
-#include "TakeRecorderSources.h"
+#include "TakePreset.h"
 #include "TakeRecorderOverlayWidget.h"
-#include "LevelSequence.h"
+#include "TakeRecorderSources.h"
+#include "TakesUtils.h"
 #include "Tickable.h"
-#include "AssetRegistryModule.h"
-#include "IAssetRegistry.h"
 #include "Stats/Stats.h"
-#include "ISequencer.h"
 #include "SequencerSettings.h"
-#include "MovieSceneTimeHelpers.h"
-#include "TakesUtils.h"
 
 // LevelSequenceEditor includes
 #include "ILevelSequenceEditorToolkit.h"
@@ -459,7 +463,8 @@ bool UTakeRecorder::Initialize ( ULevelSequence* LevelSequenceBase, UTakeRecorde
 	InitializeFromParameters();
 
 	// Open a recording notification
-	// @todo: is this too intrusive? does it potentially overlap the takerecorder slate UI?
+
+	if (ShouldShowNotifications())
 	{
 		TSharedRef<STakeRecorderNotification> Content = SNew(STakeRecorderNotification, this);
 
@@ -705,6 +710,17 @@ void UTakeRecorder::InitializeFromParameters()
 	}
 }
 
+bool UTakeRecorder::ShouldShowNotifications()
+{
+	// -TAKERECORDERISHEADLESS in the command line can force headless behavior and disable the notifications.
+	static const bool bCmdLineTakeRecorderIsHeadless = FParse::Param(FCommandLine::Get(), TEXT("TAKERECORDERISHEADLESS"));
+
+	return Parameters.Project.bShowNotifications
+		&& !bCmdLineTakeRecorderIsHeadless
+		&& !FApp::IsUnattended()
+		&& !GIsRunningUnattendedScript;
+}
+
 UWorld* UTakeRecorder::GetWorld() const
 {
 	return WeakWorld.Get();
@@ -904,6 +920,12 @@ void UTakeRecorder::Start(const FTimecode& InTimecodeSource)
 		Sources->StartRecording(SequenceAsset, InTimecodeSource, Parameters.User.bAutoSerialize ? &ManifestSerializer : nullptr);
 	}
 
+	if (!ShouldShowNotifications())
+	{
+		// Log in lieu of the notification widget
+		UE_LOG(LogTakesCore, Log, TEXT("Started recording"));
+	}
+
 	OnRecordingStartedEvent.Broadcast(this);
 
 	UTakeRecorderBlueprintLibrary::OnTakeRecorderStarted();
@@ -960,6 +982,12 @@ void UTakeRecorder::Stop()
 
 	if (bDidEverStartRecording)
 	{
+		if (!ShouldShowNotifications())
+		{
+			// Log in lieu of the notification widget
+			UE_LOG(LogTakesCore, Log, TEXT("Stopped recording"));
+		}
+
 		UTakeRecorderBlueprintLibrary::OnTakeRecorderStopped();
 
 		if (MovieScene)
@@ -1108,13 +1136,16 @@ void UTakeRecorder::HandlePIE(bool bIsSimulating)
 {
 	ULevelSequence* FinishedAsset = GetSequence();
 
-	TSharedRef<STakeRecorderNotification> Content = SNew(STakeRecorderNotification, this, FinishedAsset);
+	if (ShouldShowNotifications())
+	{
+		TSharedRef<STakeRecorderNotification> Content = SNew(STakeRecorderNotification, this, FinishedAsset);
 
-	FNotificationInfo Info(Content);
-	Info.ExpireDuration = 5.f;
+		FNotificationInfo Info(Content);
+		Info.ExpireDuration = 5.f;
 
-	TSharedPtr<SNotificationItem> PendingNotification = FSlateNotificationManager::Get().AddNotification(Info);
-	PendingNotification->SetCompletionState(SNotificationItem::CS_Success);
+		TSharedPtr<SNotificationItem> PendingNotification = FSlateNotificationManager::Get().AddNotification(Info);
+		PendingNotification->SetCompletionState(SNotificationItem::CS_Success);
+	}
 	
 	Stop();
 
diff --git a/Engine/Plugins/VirtualProduction/Takes/Source/TakeRecorder/Private/Recorder/TakeRecorderParameters.cpp b/Engine/Plugins/VirtualProduction/Takes/Source/TakeRecorder/Private/Recorder/TakeRecorderParameters.cpp
--- a/Engine/Plugins/VirtualProduction/Takes/Source/TakeRecorder/Private/Recorder/TakeRecorderParameters.cpp
+++ b/Engine/Plugins/VirtualProduction/Takes/Source/TakeRecorder/Private/Recorder/TakeRecorderParameters.cpp
@@ -24,6 +24,7 @@ FTakeRecorderProjectParameters::FTakeRecorderProjectParameters()
 	, bRecordTimecode(false)
 	, bRecordSourcesIntoSubSequences(false)
 	, bRecordToPossessable(false)
+	, bShowNotifications(true)
 {}
 
 FTakeRecorderParameters::FTakeRecorderParameters() 
diff --git a/Engine/Plugins/VirtualProduction/Takes/Source/TakeRecorder/Public/Recorder/TakeRecorder.h b/Engine/Plugins/VirtualProduction/Takes/Source/TakeRecorder/Public/Recorder/TakeRecorder.h
--- a/Engine/Plugins/VirtualProduction/Takes/Source/TakeRecorder/Public/Recorder/TakeRecorder.h
+++ b/Engine/Plugins/VirtualProduction/Takes/Source/TakeRecorder/Public/Recorder/TakeRecorder.h
@@ -181,6 +181,12 @@ private:
 	 */
 	void InitializeFromParameters();
 
+	/**
+     * Returns true if notification widgets should be shown when recording.
+	 * It takes into account TakeRecorder project settings, the command line, and global unattended settings.
+     */
+	bool ShouldShowNotifications();
+
 private:
 
 	/** Called by Tick and Start to make sure we record at start */
diff --git a/Engine/Plugins/VirtualProduction/Takes/Source/TakeRecorder/Public/Recorder/TakeRecorderParameters.h b/Engine/Plugins/VirtualProduction/Takes/Source/TakeRecorder/Public/Recorder/TakeRecorderParameters.h
--- a/Engine/Plugins/VirtualProduction/Takes/Source/TakeRecorder/Public/Recorder/TakeRecorderParameters.h
+++ b/Engine/Plugins/VirtualProduction/Takes/Source/TakeRecorder/Public/Recorder/TakeRecorderParameters.h
@@ -115,6 +115,10 @@ struct FTakeRecorderProjectParameters
 	/** List of property names for which movie scene tracks will always record. */
 	UPROPERTY(config, EditAnywhere, Category = "Take Recorder")
 	TArray<FTakeRecorderTrackSettings> DefaultTracks;
+
+	/** Whether to show notification windows or not when recording */
+	UPROPERTY(config, EditAnywhere, BlueprintReadWrite, Category = "Take Recorder")
+	bool bShowNotifications;
 };
 
 UENUM(BlueprintType)
