Add support to take recorder to allow naming override for cases where
name conflict can occur.  This can happen in Multi-user scenarios where
multiple nodes are recording take data.

diff --git a/Engine/Plugins/VirtualProduction/Takes/Source/TakeRecorder/Private/Recorder/TakeRecorder.cpp b/Engine/Plugins/VirtualProduction/Takes/Source/TakeRecorder/Private/Recorder/TakeRecorder.cpp
--- a/Engine/Plugins/VirtualProduction/Takes/Source/TakeRecorder/Private/Recorder/TakeRecorder.cpp
+++ b/Engine/Plugins/VirtualProduction/Takes/Source/TakeRecorder/Private/Recorder/TakeRecorder.cpp
@@ -2,6 +2,7 @@
 
 #include "Recorder/TakeRecorder.h"
 
+#include "Algo/Accumulate.h"
 #include "AssetRegistryModule.h"
 #include "CoreGlobals.h"
 #include "Engine/Engine.h"
@@ -358,6 +359,16 @@ static TStrongObjectPtr<UTakeRecorder>& GetCurrentRecorder()
 }
 FOnTakeRecordingInitialized UTakeRecorder::OnRecordingInitializedEvent;
 
+void FTakeRecorderParameterOverride::RegisterHandler(FName OverrideName, FTakeRecorderParameterDelegate Delegate)
+{
+	Delegates.FindOrAdd(MoveTemp(OverrideName), MoveTemp(Delegate));
+}
+
+void  FTakeRecorderParameterOverride::UnregisterHandler(FName OverrideName)
+{
+	Delegates.Remove(MoveTemp(OverrideName));
+}
+
 // Static functions for UTakeRecorder
 UTakeRecorder* UTakeRecorder::GetActiveRecorder()
 {
@@ -369,6 +380,12 @@ FOnTakeRecordingInitialized& UTakeRecorder::OnRecordingInitialized()
 	return OnRecordingInitializedEvent;
 }
 
+FTakeRecorderParameterOverride& UTakeRecorder::TakeInitializeParameterOverride()
+{
+	static FTakeRecorderParameterOverride Overrides;
+	return Overrides;
+}
+
 bool UTakeRecorder::SetActiveRecorder(UTakeRecorder* NewActiveRecorder)
 {
 	if (GetCurrentRecorder().IsValid())
@@ -397,6 +414,19 @@ void UTakeRecorder::SetDisableSaveTick(bool InValue)
 	Parameters.bDisableRecordingAndSave = InValue;
 }
 
+namespace TakeInitHelper
+{
+FTakeRecorderParameters AccumulateParamsOverride(const FTakeRecorderParameters& InParam)
+{
+	TMap<FName, FTakeRecorderParameterDelegate>& TheDelegates = UTakeRecorder::TakeInitializeParameterOverride().Delegates;
+	auto Op = [](FTakeRecorderParameters InParameters, const TPair<FName,FTakeRecorderParameterDelegate>& Pair)
+	{
+		return Pair.Value.Execute(MoveTemp(InParameters));
+	};
+	return Algo::Accumulate(TheDelegates, InParam, MoveTemp(Op));
+}
+}
+
 bool UTakeRecorder::Initialize ( ULevelSequence* LevelSequenceBase, UTakeRecorderSources* Sources, UTakeMetaData* MetaData, const FTakeRecorderParameters& InParameters, FText* OutError )
 {
 	FGCObjectScopeGuard GCGuard(this);
@@ -423,16 +453,17 @@ bool UTakeRecorder::Initialize ( ULevelSequence* LevelSequenceBase, UTakeRecorde
 
 	UTakeRecorderBlueprintLibrary::OnTakeRecorderPreInitialize();
 
-	if (InParameters.TakeRecorderMode == ETakeRecorderMode::RecordNewSequence)
+	FTakeRecorderParameters FinalParameters = TakeInitHelper::AccumulateParamsOverride(InParameters);
+	if (FinalParameters.TakeRecorderMode == ETakeRecorderMode::RecordNewSequence)
 	{
-		if (!CreateDestinationAsset(*InParameters.Project.GetTakeAssetPath(), LevelSequenceBase, Sources, MetaData, OutError))
+		if (!CreateDestinationAsset(*FinalParameters.Project.GetTakeAssetPath(), LevelSequenceBase, Sources, MetaData, OutError))
 		{
 			return false;
 		}
 	}
 	else
 	{
-		if (!SetupDestinationAsset(InParameters, LevelSequenceBase, Sources, MetaData, OutError))
+		if (!SetupDestinationAsset(FinalParameters, LevelSequenceBase, Sources, MetaData, OutError))
 		{
 			return false;
 		}
@@ -444,7 +475,7 @@ bool UTakeRecorder::Initialize ( ULevelSequence* LevelSequenceBase, UTakeRecorde
 
 	AddToRoot();
 
-	Parameters = InParameters;
+	Parameters = FinalParameters;
 	State      = ETakeRecorderState::CountingDown;
 
 	// Override parameters for recording into a current sequence
diff --git a/Engine/Plugins/VirtualProduction/Takes/Source/TakeRecorder/Public/Recorder/TakeRecorder.h b/Engine/Plugins/VirtualProduction/Takes/Source/TakeRecorder/Public/Recorder/TakeRecorder.h
--- a/Engine/Plugins/VirtualProduction/Takes/Source/TakeRecorder/Public/Recorder/TakeRecorder.h
+++ b/Engine/Plugins/VirtualProduction/Takes/Source/TakeRecorder/Public/Recorder/TakeRecorder.h
@@ -37,6 +37,15 @@ DECLARE_MULTICAST_DELEGATE_OneParam(FOnTakeRecordingStarted, UTakeRecorder*);
 DECLARE_MULTICAST_DELEGATE_OneParam(FOnTakeRecordingFinished, UTakeRecorder*);
 DECLARE_MULTICAST_DELEGATE_OneParam(FOnTakeRecordingCancelled, UTakeRecorder*);
 
+DECLARE_DELEGATE_RetVal_OneParam(FTakeRecorderParameters, FTakeRecorderParameterDelegate,const FTakeRecorderParameters&);
+
+struct TAKERECORDER_API FTakeRecorderParameterOverride
+{
+	void RegisterHandler(FName OverrideName, FTakeRecorderParameterDelegate Delegate);
+	void UnregisterHandler(FName OverrideName);
+
+	TMap<FName, FTakeRecorderParameterDelegate> Delegates;
+};
 
 DECLARE_LOG_CATEGORY_EXTERN(ManifestSerialization, Verbose, All);
 
@@ -61,6 +70,12 @@ public:
 	 */
 	static FOnTakeRecordingInitialized& OnRecordingInitialized();
 
+	/**
+	 * On take initialization overrides provide a mechanism to adjust parameter values when
+	 * a new take is initialized.
+	 */
+	static FTakeRecorderParameterOverride& TakeInitializeParameterOverride();
+
 public:
 
 	/**
