Take Recorder: Cancel recording

When cancelled, sources and recorder have an opportunity to clean up assets that may have been created.
  - MicrophoneAudioSource removes recorded sound wave assets.
  - AnimationTrackRecorder removes recorded animation assets.
  - SubSequence assets are removed on cancel

Added dependency to UnrealEd from TakeTrackRecorders for ObjectTools::DeleteAssets

diff --git a/Engine/Plugins/Animation/LiveLink/Source/LiveLinkSequencer/Private/TakeRecorderSource/TakeRecorderLiveLinkSource.cpp b/Engine/Plugins/Animation/LiveLink/Source/LiveLinkSequencer/Private/TakeRecorderSource/TakeRecorderLiveLinkSource.cpp
--- a/Engine/Plugins/Animation/LiveLink/Source/LiveLinkSequencer/Private/TakeRecorderSource/TakeRecorderLiveLinkSource.cpp
+++ b/Engine/Plugins/Animation/LiveLink/Source/LiveLinkSequencer/Private/TakeRecorderSource/TakeRecorderLiveLinkSource.cpp
@@ -63,7 +63,7 @@ void UTakeRecorderLiveLinkSource::StopRecording(class ULevelSequence* InSequence
 	}
 }
 
-TArray<UTakeRecorderSource*> UTakeRecorderLiveLinkSource::PostRecording(class ULevelSequence* InSequence, class ULevelSequence* InMasterSequence)
+TArray<UTakeRecorderSource*> UTakeRecorderLiveLinkSource::PostRecording(class ULevelSequence* InSequence, class ULevelSequence* InMasterSequence, const bool bCancelled)
 {
 	if (TrackRecorder)
 	{
diff --git a/Engine/Plugins/Animation/LiveLink/Source/LiveLinkSequencer/Private/TakeRecorderSource/TakeRecorderLiveLinkSource.h b/Engine/Plugins/Animation/LiveLink/Source/LiveLinkSequencer/Private/TakeRecorderSource/TakeRecorderLiveLinkSource.h
--- a/Engine/Plugins/Animation/LiveLink/Source/LiveLinkSequencer/Private/TakeRecorderSource/TakeRecorderLiveLinkSource.h
+++ b/Engine/Plugins/Animation/LiveLink/Source/LiveLinkSequencer/Private/TakeRecorderSource/TakeRecorderLiveLinkSource.h
@@ -88,7 +88,7 @@ private:
 	virtual void StartRecording(const FTimecode& InSectionStartTimecode, const FFrameNumber& InSectionFirstFrame, class ULevelSequence* InSequence) override;
 	virtual void TickRecording(const FQualifiedFrameTime& CurrentSequenceTime) override;
 	virtual void StopRecording(class ULevelSequence* InSequence) override;
-	virtual TArray<UTakeRecorderSource*> PostRecording(class ULevelSequence* InSequence, class ULevelSequence* InMasterSequence) override;
+	virtual TArray<UTakeRecorderSource*> PostRecording(class ULevelSequence* InSequence, class ULevelSequence* InMasterSequence, const bool bCancelled) override;
 
 	virtual FText GetDisplayTextImpl() const override;
 	virtual void AddContentsToFolder(class UMovieSceneFolder* InFolder) override;
diff --git a/Engine/Plugins/VirtualProduction/DMX/DMXEngine/Source/DMXEditor/Private/Sequencer/TakeRecorderDMXLibrarySource.cpp b/Engine/Plugins/VirtualProduction/DMX/DMXEngine/Source/DMXEditor/Private/Sequencer/TakeRecorderDMXLibrarySource.cpp
--- a/Engine/Plugins/VirtualProduction/DMX/DMXEngine/Source/DMXEditor/Private/Sequencer/TakeRecorderDMXLibrarySource.cpp
+++ b/Engine/Plugins/VirtualProduction/DMX/DMXEngine/Source/DMXEditor/Private/Sequencer/TakeRecorderDMXLibrarySource.cpp
@@ -102,7 +102,7 @@ void UTakeRecorderDMXLibrarySource::StopRecording(class ULevelSequence* InSequen
 	}
 }
 
-TArray<UTakeRecorderSource*> UTakeRecorderDMXLibrarySource::PostRecording(class ULevelSequence* InSequence, ULevelSequence* InMasterSequence)
+TArray<UTakeRecorderSource*> UTakeRecorderDMXLibrarySource::PostRecording(class ULevelSequence* InSequence, ULevelSequence* InMasterSequence, const bool bCancelled)
 {
 	if (TrackRecorder)
 	{
diff --git a/Engine/Plugins/VirtualProduction/DMX/DMXEngine/Source/DMXEditor/Private/Sequencer/TakeRecorderDMXLibrarySource.h b/Engine/Plugins/VirtualProduction/DMX/DMXEngine/Source/DMXEditor/Private/Sequencer/TakeRecorderDMXLibrarySource.h
--- a/Engine/Plugins/VirtualProduction/DMX/DMXEngine/Source/DMXEditor/Private/Sequencer/TakeRecorderDMXLibrarySource.h
+++ b/Engine/Plugins/VirtualProduction/DMX/DMXEngine/Source/DMXEditor/Private/Sequencer/TakeRecorderDMXLibrarySource.h
@@ -81,7 +81,7 @@ private:
 	virtual void StartRecording(const FTimecode& InSectionStartTimecode, const FFrameNumber& InSectionFirstFrame, class ULevelSequence* InSequence) override;
 	virtual void StopRecording(class ULevelSequence* InSequence) override;
 	virtual void TickRecording(const FQualifiedFrameTime& CurrentTime) override;
-	virtual TArray<UTakeRecorderSource*> PostRecording(class ULevelSequence* InSequence, ULevelSequence* InMasterSequence) override;
+	virtual TArray<UTakeRecorderSource*> PostRecording(class ULevelSequence* InSequence, ULevelSequence* InMasterSequence, const bool bCancelled) override;
 	virtual void AddContentsToFolder(UMovieSceneFolder* InFolder) override;
 	virtual FText GetDisplayTextImpl() const override;
 	// ~End UTakeRecorderSource
diff --git a/Engine/Plugins/VirtualProduction/Takes/Source/TakeRecorder/Private/Recorder/TakeRecorder.cpp b/Engine/Plugins/VirtualProduction/Takes/Source/TakeRecorder/Private/Recorder/TakeRecorder.cpp
--- a/Engine/Plugins/VirtualProduction/Takes/Source/TakeRecorder/Private/Recorder/TakeRecorder.cpp
+++ b/Engine/Plugins/VirtualProduction/Takes/Source/TakeRecorder/Private/Recorder/TakeRecorder.cpp
@@ -13,6 +13,7 @@
 #include "Misc/CommandLine.h"
 #include "Misc/Parse.h"
 #include "MovieSceneTimeHelpers.h"
+#include "ObjectTools.h"
 #include "Recorder/TakeRecorderBlueprintLibrary.h"
 #include "TakeMetaData.h"
 #include "TakePreset.h"
@@ -1018,6 +1019,18 @@ void UTakeRecorder::Start()
 }
 
 void UTakeRecorder::Stop()
+{
+	const bool bCancelled = false;
+	StopInternal(bCancelled);
+}
+
+void UTakeRecorder::Cancel()
+{
+	const bool bCancelled = true;
+	StopInternal(bCancelled);
+}
+
+void UTakeRecorder::StopInternal(const bool bCancelled)
 {
 	static bool bStoppedRecording = false;
 
@@ -1040,12 +1053,12 @@ void UTakeRecorder::Stop()
 
 	ManifestSerializer.Close();
 
-	const bool bDidEverStartRecording = State == ETakeRecorderState::Started;
-
 	FEditorDelegates::EndPIE.RemoveAll(this);
 	FEditorDelegates::BeginPIE.RemoveAll(this);
-
-	State = bDidEverStartRecording ? ETakeRecorderState::Stopped : ETakeRecorderState::Cancelled;
+	
+	const bool bDidEverStartRecording = State == ETakeRecorderState::Started;
+	const bool bRecordingFinished = !bCancelled && bDidEverStartRecording;
+	State = bRecordingFinished ? ETakeRecorderState::Stopped : ETakeRecorderState::Cancelled;
 
 	TSharedPtr<ISequencer> Sequencer = WeakSequencer.Pin();
 	if (Sequencer.IsValid())
@@ -1069,7 +1082,14 @@ void UTakeRecorder::Stop()
 		if (!ShouldShowNotifications())
 		{
 			// Log in lieu of the notification widget
-			UE_LOG(LogTakesCore, Log, TEXT("Stopped recording"));
+			if (bRecordingFinished)
+			{
+				UE_LOG(LogTakesCore, Log, TEXT("Stopped recording"));
+			}
+			else
+			{
+				UE_LOG(LogTakesCore, Log, TEXT("Recording cancelled"));
+			}
 		}
 
 		UTakeRecorderBlueprintLibrary::OnTakeRecorderStopped();
@@ -1098,7 +1118,7 @@ void UTakeRecorder::Stop()
 		{
 			UTakeRecorderSources* Sources = SequenceAsset->FindMetaData<UTakeRecorderSources>();
 			check(Sources);
-			Sources->StopRecording(SequenceAsset);
+			Sources->StopRecording(SequenceAsset, bCancelled);
 		}
 
 		// Restore the playback range to what it was before recording.
@@ -1115,44 +1135,48 @@ void UTakeRecorder::Stop()
 			TakesUtils::ClampPlaybackRangeToEncompassAllSections(SequenceAsset->GetMovieScene(), bUpperBoundOnly);
 		}
 
-		// Lock the sequence so that it can't be changed without implicitly unlocking it now
-		if (Parameters.User.bAutoLock)
+		if (bRecordingFinished)
 		{
-			SequenceAsset->GetMovieScene()->SetReadOnly(true);
-		}
+			// Lock the sequence so that it can't be changed without implicitly unlocking it now
+			if (Parameters.User.bAutoLock)
+			{
+				SequenceAsset->GetMovieScene()->SetReadOnly(true);
+			}
 
-		UTakeMetaData* AssetMetaData = SequenceAsset->FindMetaData<UTakeMetaData>();
-		check(AssetMetaData);
+			UTakeMetaData* AssetMetaData = SequenceAsset->FindMetaData<UTakeMetaData>();
+			check(AssetMetaData);
 
-		if (MovieScene)
-		{
-			FFrameRate DisplayRate = MovieScene->GetDisplayRate();
-			FFrameRate TickResolution = MovieScene->GetTickResolution();
-			FTimecode Timecode = FTimecode::FromFrameNumber(FFrameRate::TransformTime(CurrentFrameTime, TickResolution, DisplayRate).FloorToFrame(), DisplayRate);
+			if (MovieScene)
+			{
+				FFrameRate DisplayRate = MovieScene->GetDisplayRate();
+				FFrameRate TickResolution = MovieScene->GetTickResolution();
+				FTimecode Timecode = FTimecode::FromFrameNumber(FFrameRate::TransformTime(CurrentFrameTime, TickResolution, DisplayRate).FloorToFrame(), DisplayRate);
 
-			AssetMetaData->SetTimecodeOut(Timecode);
-		}
+				AssetMetaData->SetTimecodeOut(Timecode);
+			}
 
-		if (GEditor && GEditor->GetEditorWorldContext().World())
-		{
-			AssetMetaData->SetLevelOrigin(GEditor->GetEditorWorldContext().World()->PersistentLevel);
-		}
+			if (GEditor && GEditor->GetEditorWorldContext().World())
+			{
+				AssetMetaData->SetLevelOrigin(GEditor->GetEditorWorldContext().World()->PersistentLevel);
+			}
 
-		// Lock the meta data so it can't be changed without implicitly unlocking it now
-		AssetMetaData->Lock();
+			// Lock the meta data so it can't be changed without implicitly unlocking it now
+			AssetMetaData->Lock();
 
-		if (Parameters.User.bSaveRecordedAssets)
-		{
-			TakesUtils::SaveAsset(SequenceAsset);
-		}
-		
-		// Rebuild sequencer because subsequences could have been added or bindings removed
-		if (Sequencer)
-		{
-			Sequencer->RefreshTree();
+			if (Parameters.User.bSaveRecordedAssets)
+			{
+				TakesUtils::SaveAsset(SequenceAsset);
+			}
+
+			// Rebuild sequencer because subsequences could have been added or bindings removed
+			if (Sequencer)
+			{
+				Sequencer->RefreshTree();
+			}
 		}
 	}
-	else if (Parameters.TakeRecorderMode == ETakeRecorderMode::RecordNewSequence)
+
+	if (Parameters.TakeRecorderMode == ETakeRecorderMode::RecordNewSequence && !bRecordingFinished)
 	{
 		if (GIsEditor)
 		{
@@ -1186,7 +1210,7 @@ void UTakeRecorder::Stop()
 		GetCurrentRecorder().Reset();
 		TickableTakeRecorder.WeakRecorder = nullptr;
 
-		if (bDidEverStartRecording)
+		if (bRecordingFinished)
 		{
 			OnRecordingFinishedEvent.Broadcast(this);
 			UTakeRecorderBlueprintLibrary::OnTakeRecorderFinished(SequenceAsset);
diff --git a/Engine/Plugins/VirtualProduction/Takes/Source/TakeRecorder/Private/Recorder/TakeRecorderBlueprintLibrary.cpp b/Engine/Plugins/VirtualProduction/Takes/Source/TakeRecorder/Private/Recorder/TakeRecorderBlueprintLibrary.cpp
--- a/Engine/Plugins/VirtualProduction/Takes/Source/TakeRecorder/Private/Recorder/TakeRecorderBlueprintLibrary.cpp
+++ b/Engine/Plugins/VirtualProduction/Takes/Source/TakeRecorder/Private/Recorder/TakeRecorderBlueprintLibrary.cpp
@@ -101,6 +101,15 @@ void UTakeRecorderBlueprintLibrary::StopRecording()
 	}
 }
 
+void UTakeRecorderBlueprintLibrary::CancelRecording()
+{
+	UTakeRecorder* Existing = UTakeRecorder::GetActiveRecorder();
+	if (Existing)
+	{
+		Existing->Cancel();
+	}
+}
+
 UTakeRecorderPanel* UTakeRecorderBlueprintLibrary::OpenTakeRecorderPanel()
 {
 	UTakeRecorderPanel* Existing = GetTakeRecorderPanel();
diff --git a/Engine/Plugins/VirtualProduction/Takes/Source/TakeRecorder/Private/Widgets/STakeRecorderCockpit.cpp b/Engine/Plugins/VirtualProduction/Takes/Source/TakeRecorder/Private/Widgets/STakeRecorderCockpit.cpp
--- a/Engine/Plugins/VirtualProduction/Takes/Source/TakeRecorder/Private/Widgets/STakeRecorderCockpit.cpp
+++ b/Engine/Plugins/VirtualProduction/Takes/Source/TakeRecorder/Private/Widgets/STakeRecorderCockpit.cpp
@@ -56,6 +56,7 @@
 #include "Modules/ModuleManager.h"
 #include "LevelEditor.h"
 #include "ISettingsModule.h"
+#include "Dialogs/CustomDialog.h"
 
 extern UNREALED_API UEditorEngine* GEditor;
 
@@ -337,6 +338,26 @@ void STakeRecorderCockpit::Construct(const FArguments& InArgs)
 						.Text(this, &STakeRecorderCockpit::GetCountdownText)
 					]
 				]
+
+				+ SHorizontalBox::Slot()
+				.AutoWidth()
+				.Padding(2)
+				[
+					SNew(SComboButton)
+					.ContentPadding(2)
+					.ForegroundColor(FSlateColor::UseForeground())
+					.ComboButtonStyle(FTakeRecorderStyle::Get(), "ComboButton")
+					.ToolTipText(LOCTEXT("RecordingOptionsTooltip", "Recording options"))
+					.OnGetMenuContent(this, &STakeRecorderCockpit::OnRecordingOptionsMenu)
+					.HasDownArrow(false)
+					.ButtonContent()
+					[
+						SNew(STextBlock)
+						.TextStyle(FEditorStyle::Get(), "NormalText.Important")
+						.Font(FEditorStyle::Get().GetFontStyle("FontAwesome.10"))
+						.Text(FEditorFontGlyphs::Caret_Down)
+					]
+				]
 			]
 		]
 
@@ -602,6 +623,23 @@ FText STakeRecorderCockpit::GetCountdownText() const
 	return bIsCountingDown ? FText::AsNumber(FMath::CeilToInt(CurrentRecording->GetCountdownSeconds())) : FText();
 }
 
+TSharedRef<SWidget> STakeRecorderCockpit::OnRecordingOptionsMenu()
+{
+	FMenuBuilder MenuBuilder(true, nullptr);
+
+	MenuBuilder.AddMenuEntry(
+		LOCTEXT("CancelRecording_Text", "Cancel Recording"),
+		LOCTEXT("CancelRecording_Tip", "Cancel the current recording, deleting any assets and resetting the take number"),
+		FSlateIcon(),
+		FUIAction(
+			FExecuteAction::CreateSP(this, &STakeRecorderCockpit::CancelRecording),
+			FCanExecuteAction::CreateLambda([this] { return Recording(); })
+		)
+	);
+
+	return MenuBuilder.MakeWidget();
+}
+
 void STakeRecorderCockpit::CacheMetaData()
 {
 	UTakeMetaData* NewMetaDataThisTick = nullptr;
@@ -963,6 +1001,30 @@ void STakeRecorderCockpit::StopRecording()
 	}
 }
 
+void STakeRecorderCockpit::CancelRecording()
+{
+	UTakeRecorder* CurrentRecording = UTakeRecorder::GetActiveRecorder();
+	if (CurrentRecording)
+	{
+		TSharedRef<SCustomDialog> ConfirmDialog = SNew(SCustomDialog)
+			.Title(FText(LOCTEXT("ConfirmCancelRecordingTitle", "Cancel Recording?")))
+			.DialogContent( SNew(STextBlock).Text(LOCTEXT("ConfirmCancelRecording", "Are you sure you want to cancel the current recording?")))
+			.Buttons({
+				SCustomDialog::FButton(LOCTEXT("Yes", "Yes"), FSimpleDelegate::CreateLambda([WeakRecording = TWeakObjectPtr<UTakeRecorder>(CurrentRecording)]()
+					{ 
+						if (WeakRecording.IsValid())
+						{
+							WeakRecording.Get()->Cancel(); 
+						}
+					})),
+				SCustomDialog::FButton(LOCTEXT("No", "No"))
+			});
+
+		// Non modal so that the recording continues to update
+		ConfirmDialog->Show();
+	}
+}
+
 void STakeRecorderCockpit::StartRecording()
 {
 	static bool bStartedRecording = false;
diff --git a/Engine/Plugins/VirtualProduction/Takes/Source/TakeRecorder/Private/Widgets/STakeRecorderCockpit.h b/Engine/Plugins/VirtualProduction/Takes/Source/TakeRecorder/Private/Widgets/STakeRecorderCockpit.h
--- a/Engine/Plugins/VirtualProduction/Takes/Source/TakeRecorder/Private/Widgets/STakeRecorderCockpit.h
+++ b/Engine/Plugins/VirtualProduction/Takes/Source/TakeRecorder/Private/Widgets/STakeRecorderCockpit.h
@@ -62,6 +62,8 @@ public:
 
 	void StopRecording();
 
+	void CancelRecording();
+
 	void Refresh();
 
 private:
@@ -87,6 +89,8 @@ private:
 	EVisibility GetCountdownVisibility() const;
 	FText GetCountdownText() const;
 
+	TSharedRef<SWidget> OnRecordingOptionsMenu();
+
 	FText GetTimecodeText() const;
 
 	FText GetUserDescriptionText() const;
diff --git a/Engine/Plugins/VirtualProduction/Takes/Source/TakeRecorder/Public/Recorder/TakeRecorder.h b/Engine/Plugins/VirtualProduction/Takes/Source/TakeRecorder/Public/Recorder/TakeRecorder.h
--- a/Engine/Plugins/VirtualProduction/Takes/Source/TakeRecorder/Public/Recorder/TakeRecorder.h
+++ b/Engine/Plugins/VirtualProduction/Takes/Source/TakeRecorder/Public/Recorder/TakeRecorder.h
@@ -129,6 +129,11 @@ public:
 	 */
 	void Stop();
 
+	/**
+	 * Called to cancel the recording
+	 */
+	void Cancel();
+
 	/**
 	 * Retrieve a multi-cast delegate that is triggered before initialization occurs (ie. when the recording button is pressed and before the countdown starts)
 	 */
@@ -161,6 +166,11 @@ private:
 	 */
 	void Start();
 
+	/*
+	 * Stop or cancel
+	 */
+	void StopInternal(const bool bCancelled);
+
 	/**
 	 * Ticked by a tickable game object to performe any necessary time-sliced logic
 	 */
diff --git a/Engine/Plugins/VirtualProduction/Takes/Source/TakeRecorder/Public/Recorder/TakeRecorderBlueprintLibrary.h b/Engine/Plugins/VirtualProduction/Takes/Source/TakeRecorder/Public/Recorder/TakeRecorderBlueprintLibrary.h
--- a/Engine/Plugins/VirtualProduction/Takes/Source/TakeRecorder/Public/Recorder/TakeRecorderBlueprintLibrary.h
+++ b/Engine/Plugins/VirtualProduction/Takes/Source/TakeRecorder/Public/Recorder/TakeRecorderBlueprintLibrary.h
@@ -66,6 +66,13 @@ public:
 	static void StopRecording();
 
 
+	/**
+	 * Cancel recording if there is a recorder currently active
+	 */
+	UFUNCTION(BlueprintCallable, Category="Take Recorder")
+	static void CancelRecording();
+
+
 	/**
 	 * Get the currently open take recorder panel, if one is open
 	 */
diff --git a/Engine/Plugins/VirtualProduction/Takes/Source/TakeRecorderSources/Private/TakeRecorderActorSource.cpp b/Engine/Plugins/VirtualProduction/Takes/Source/TakeRecorderSources/Private/TakeRecorderActorSource.cpp
--- a/Engine/Plugins/VirtualProduction/Takes/Source/TakeRecorderSources/Private/TakeRecorderActorSource.cpp
+++ b/Engine/Plugins/VirtualProduction/Takes/Source/TakeRecorderSources/Private/TakeRecorderActorSource.cpp
@@ -741,7 +741,7 @@ void UTakeRecorderActorSource::ProcessRecordedTimes(ULevelSequence* InSequence)
 	}
 }
 
-TArray<UTakeRecorderSource*> UTakeRecorderActorSource::PostRecording(ULevelSequence* InSequence, class ULevelSequence* InMasterSequence)
+TArray<UTakeRecorderSource*> UTakeRecorderActorSource::PostRecording(ULevelSequence* InSequence, class ULevelSequence* InMasterSequence, const bool bCancelled)
 {
 	FTakeRecorderParameters Parameters;
 	Parameters.User = GetDefault<UTakeRecorderUserSettings>()->Settings;
@@ -765,33 +765,43 @@ TArray<UTakeRecorderSource*> UTakeRecorderActorSource::PostRecording(ULevelSeque
 
 		// takerecorder-todo: Section Recorders should have display names, update this to use those.
 		SlowTask.EnterProgressFrame(1.0f, FText::Format(LOCTEXT("FinalizingTrackRecorder", "Finalizing Section Recorder {0}/{1}"), SectionRecorderIndex, TrackRecorders.Num()));
-		SectionRecorder->FinalizeTrack();
-	}
-
-	if (Parameters.Project.bRecordTimecode)
-	{
-		ProcessRecordedTimes(InSequence);
+		if (bCancelled)
+		{
+			SectionRecorder->CancelTrack();
+		}
+		else
+		{
+			SectionRecorder->FinalizeTrack();
+		}
 	}
 
-	// Expand the Movie Scene Playback Range to encompass all of the sections now that they've all been created.
-	SequenceRecorderUtils::ExtendSequencePlaybackRange(InSequence);
-
-	if (Target.IsValid())
+	if (!bCancelled)
 	{
-		// Automatically add or update the camera cut track if there is a camera component
-		AActor* TargetActor = Target.Get();
-
-		if (TargetActor->GetComponentByClass(UCameraComponent::StaticClass()))
+		if (Parameters.Project.bRecordTimecode)
 		{
-			FGuid RecordedCameraGuid = GetRecordedActorGuid(Target.Get());
-			FMovieSceneSequenceID RecordedCameraSequenceID = GetLevelSequenceID(Target.Get());
-			TakesUtils::CreateCameraCutTrack(InMasterSequence, RecordedCameraGuid, RecordedCameraSequenceID, InSequence->GetMovieScene()->GetPlaybackRange());
+			ProcessRecordedTimes(InSequence);
 		}
 
-		// Swap our target actor to the Editor actor (in case the recording was added while in PIE)
-		if (AActor* EditorActor = EditorUtilities::GetEditorWorldCounterpartActor(Target.Get()))
+		// Expand the Movie Scene Playback Range to encompass all of the sections now that they've all been created.
+		SequenceRecorderUtils::ExtendSequencePlaybackRange(InSequence);
+
+		if (Target.IsValid())
 		{
-			Target = EditorActor;
+			// Automatically add or update the camera cut track if there is a camera component
+			AActor* TargetActor = Target.Get();
+
+			if (TargetActor->GetComponentByClass(UCameraComponent::StaticClass()))
+			{
+				FGuid RecordedCameraGuid = GetRecordedActorGuid(Target.Get());
+				FMovieSceneSequenceID RecordedCameraSequenceID = GetLevelSequenceID(Target.Get());
+				TakesUtils::CreateCameraCutTrack(InMasterSequence, RecordedCameraGuid, RecordedCameraSequenceID, InSequence->GetMovieScene()->GetPlaybackRange());
+			}
+
+			// Swap our target actor to the Editor actor (in case the recording was added while in PIE)
+			if (AActor* EditorActor = EditorUtilities::GetEditorWorldCounterpartActor(Target.Get()))
+			{
+				Target = EditorActor;
+			}
 		}
 	}
 
diff --git a/Engine/Plugins/VirtualProduction/Takes/Source/TakeRecorderSources/Private/TakeRecorderActorSource.h b/Engine/Plugins/VirtualProduction/Takes/Source/TakeRecorderSources/Private/TakeRecorderActorSource.h
--- a/Engine/Plugins/VirtualProduction/Takes/Source/TakeRecorderSources/Private/TakeRecorderActorSource.h
+++ b/Engine/Plugins/VirtualProduction/Takes/Source/TakeRecorderSources/Private/TakeRecorderActorSource.h
@@ -124,7 +124,7 @@ public:
 	virtual void StartRecording(const FTimecode& InSectionStartTimecode, const FFrameNumber& InSectionFirstFrame, class ULevelSequence* InSequence) override;
 	virtual void TickRecording(const FQualifiedFrameTime& CurrentSequenceTime) override;
 	virtual void StopRecording(class ULevelSequence* InSequence) override;
-	virtual TArray<UTakeRecorderSource*> PostRecording(class ULevelSequence* InSequence, class ULevelSequence* InMasterSequence) override;
+	virtual TArray<UTakeRecorderSource*> PostRecording(class ULevelSequence* InSequence, class ULevelSequence* InMasterSequence, const bool bCancelled) override;
 	virtual TArray<UObject*> GetAdditionalSettingsObjects() const { return TArray<UObject*>(FactorySettings); }
 	virtual FString GetSubsceneTrackName(ULevelSequence* InSequence) const override;
 	virtual FString GetSubsceneAssetName(ULevelSequence* InSequence) const override;
diff --git a/Engine/Plugins/VirtualProduction/Takes/Source/TakeRecorderSources/Private/TakeRecorderCameraCutSource.cpp b/Engine/Plugins/VirtualProduction/Takes/Source/TakeRecorderSources/Private/TakeRecorderCameraCutSource.cpp
--- a/Engine/Plugins/VirtualProduction/Takes/Source/TakeRecorderSources/Private/TakeRecorderCameraCutSource.cpp
+++ b/Engine/Plugins/VirtualProduction/Takes/Source/TakeRecorderSources/Private/TakeRecorderCameraCutSource.cpp
@@ -100,10 +100,10 @@ void UTakeRecorderCameraCutSource::TickRecording(const FQualifiedFrameTime& Curr
 }
 
 
-TArray<UTakeRecorderSource*> UTakeRecorderCameraCutSource::PostRecording(class ULevelSequence* InSequence, class ULevelSequence* InMasterSequence)
+TArray<UTakeRecorderSource*> UTakeRecorderCameraCutSource::PostRecording(class ULevelSequence* InSequence, class ULevelSequence* InMasterSequence, const bool bCancelled)
 {
 	// Build the camera cut track
-	if (CameraCutData.Num())
+	if (!bCancelled && CameraCutData.Num())
 	{
 		UMovieSceneTrack* CameraCutTrack = InMasterSequence->GetMovieScene()->GetCameraCutTrack();
 		if (!CameraCutTrack)
diff --git a/Engine/Plugins/VirtualProduction/Takes/Source/TakeRecorderSources/Private/TakeRecorderCameraCutSource.h b/Engine/Plugins/VirtualProduction/Takes/Source/TakeRecorderSources/Private/TakeRecorderCameraCutSource.h
--- a/Engine/Plugins/VirtualProduction/Takes/Source/TakeRecorderSources/Private/TakeRecorderCameraCutSource.h
+++ b/Engine/Plugins/VirtualProduction/Takes/Source/TakeRecorderSources/Private/TakeRecorderCameraCutSource.h
@@ -27,7 +27,7 @@ private:
 	// UTakeRecorderSource
 	virtual TArray<UTakeRecorderSource*> PreRecording(ULevelSequence* InSequence, FMovieSceneSequenceID InSequenceID, ULevelSequence* InMasterSequence, FManifestSerializer* InManifestSerializer) override;
 	virtual void TickRecording(const FQualifiedFrameTime& CurrentTime) override;
-	virtual TArray<UTakeRecorderSource*> PostRecording(class ULevelSequence* InSequence, class ULevelSequence* InMasterSequence) override;
+	virtual TArray<UTakeRecorderSource*> PostRecording(class ULevelSequence* InSequence, class ULevelSequence* InMasterSequence, const bool bCancelled) override;
 	virtual FText GetDisplayTextImpl() const override;
 
 	// This source does not support subscenes
diff --git a/Engine/Plugins/VirtualProduction/Takes/Source/TakeRecorderSources/Private/TakeRecorderMicrophoneAudioSource.cpp b/Engine/Plugins/VirtualProduction/Takes/Source/TakeRecorderSources/Private/TakeRecorderMicrophoneAudioSource.cpp
--- a/Engine/Plugins/VirtualProduction/Takes/Source/TakeRecorderSources/Private/TakeRecorderMicrophoneAudioSource.cpp
+++ b/Engine/Plugins/VirtualProduction/Takes/Source/TakeRecorderSources/Private/TakeRecorderMicrophoneAudioSource.cpp
@@ -21,6 +21,7 @@
 #include "AssetData.h"
 #include "AssetRegistryModule.h"
 #include "UObject/UObjectBaseUtility.h"
+#include "ObjectTools.h"
 
 UTakeRecorderMicrophoneAudioSourceSettings::UTakeRecorderMicrophoneAudioSourceSettings(const FObjectInitializer& ObjInit)
 	: Super(ObjInit)
@@ -145,6 +146,8 @@ void UTakeRecorderMicrophoneAudioSource::StartRecording(const FTimecode& InSecti
 	AudioSettings.GainDb = AudioGain;
 	AudioSettings.bSplitChannels = bSplitAudioChannelsIntoSeparateTracks;
 
+	RecordedSoundWaves.Empty();
+	
 	AudioRecorder = Recorder.CreateAudioRecorder();
 	if (AudioRecorder)
 	{
@@ -161,64 +164,97 @@ void UTakeRecorderMicrophoneAudioSource::StopRecording(class ULevelSequence* InS
 {
 	Super::StopRecording(InSequence);
 
-	TArray<USoundWave*> RecordedSoundWaves;
 	if (AudioRecorder)
 	{
-		AudioRecorder->Stop(RecordedSoundWaves);
+		TArray<USoundWave*> SoundWaves;
+		AudioRecorder->Stop(SoundWaves);
+
+		for (USoundWave* RecordedSoundWave : SoundWaves)
+		{
+			RecordedSoundWaves.Add(RecordedSoundWave);
+		}
 	}
 	AudioRecorder.Reset();
+}
 
+TArray<UTakeRecorderSource*> UTakeRecorderMicrophoneAudioSource::PostRecording(ULevelSequence* InSequence, class ULevelSequence* InMasterSequence, const bool bCancelled)
+{
 	if (!RecordedSoundWaves.Num())
 	{
-		return;
+		return TArray<UTakeRecorderSource*>();
 	}
 
-	for (auto RecordedSoundWave : RecordedSoundWaves)
+	TArray<UObject*> AssetsToCleanUp;
+	if (bCancelled)
 	{
-		RecordedSoundWave->MarkPackageDirty();
-		
-		FAssetRegistryModule::AssetCreated(RecordedSoundWave);
+		for (TWeakObjectPtr<USoundWave> WeakRecordedSoundWave : RecordedSoundWaves)
+		{
+			if (USoundWave* RecordedSoundWave = WeakRecordedSoundWave.Get())
+			{
+				AssetsToCleanUp.Add(RecordedSoundWave);
+			}
+		}
 	}
+	else
+	{
+		for (TWeakObjectPtr<USoundWave> WeakRecordedSoundWave : RecordedSoundWaves)
+		{
+			if (USoundWave* RecordedSoundWave = WeakRecordedSoundWave.Get())
+			{
+				RecordedSoundWave->MarkPackageDirty();
+		
+				FAssetRegistryModule::AssetCreated(RecordedSoundWave);
+			}
+		}
 
-	UMovieScene* MovieScene = InSequence->GetMovieScene();
-	check(CachedAudioTrack.IsValid());
+		UMovieScene* MovieScene = InSequence->GetMovieScene();
+		check(CachedAudioTrack.IsValid());
 
-	FFrameRate TickResolution = MovieScene->GetTickResolution();
-	FFrameRate DisplayRate = MovieScene->GetDisplayRate();
+		FFrameRate TickResolution = MovieScene->GetTickResolution();
+		FFrameRate DisplayRate = MovieScene->GetDisplayRate();
 
-	if (bReplaceRecordedAudio)
-	{
-		CachedAudioTrack->RemoveAllAnimationData();
-	}
+		if (bReplaceRecordedAudio)
+		{
+			CachedAudioTrack->RemoveAllAnimationData();
+		}
 
-	UTakeRecorderSources* Sources = GetTypedOuter<UTakeRecorderSources>();
+		UTakeRecorderSources* Sources = GetTypedOuter<UTakeRecorderSources>();
 
-	for (USoundWave* RecordedAudio : RecordedSoundWaves)
-	{
-		int32 RowIndex = -1;
-		for (UMovieSceneSection* Section : CachedAudioTrack->GetAllSections())
+		for (TWeakObjectPtr<USoundWave> WeakRecordedSoundWave : RecordedSoundWaves)
 		{
-			RowIndex = FMath::Max(RowIndex, Section->GetRowIndex());
-		}
+			if (USoundWave* RecordedSoundWave = WeakRecordedSoundWave.Get())
+			{
+				int32 RowIndex = -1;
+				for (UMovieSceneSection* Section : CachedAudioTrack->GetAllSections())
+				{
+					RowIndex = FMath::Max(RowIndex, Section->GetRowIndex());
+				}
 
-		UMovieSceneAudioSection* NewAudioSection = NewObject<UMovieSceneAudioSection>(CachedAudioTrack.Get(), UMovieSceneAudioSection::StaticClass());
+				UMovieSceneAudioSection* NewAudioSection = NewObject<UMovieSceneAudioSection>(CachedAudioTrack.Get(), UMovieSceneAudioSection::StaticClass());
 
-		FFrameNumber RecordStartFrame = MovieScene->GetPlaybackRange().GetLowerBoundValue();
+				FFrameNumber RecordStartFrame = MovieScene->GetPlaybackRange().GetLowerBoundValue();
 
-		NewAudioSection->SetSound(RecordedAudio);
-		NewAudioSection->SetRange(TRange<FFrameNumber>(RecordStartFrame, RecordStartFrame + (RecordedAudio->GetDuration() * TickResolution).CeilToFrame()));
-		NewAudioSection->TimecodeSource = FTimecode::FromFrameNumber(RecordStartFrame, DisplayRate);
+				NewAudioSection->SetSound(RecordedSoundWave);
+				NewAudioSection->SetRange(TRange<FFrameNumber>(RecordStartFrame, RecordStartFrame + (RecordedSoundWave->GetDuration() * TickResolution).CeilToFrame()));
+				NewAudioSection->TimecodeSource = FTimecode::FromFrameNumber(RecordStartFrame, DisplayRate);
 
-		CachedAudioTrack->AddSection(*NewAudioSection);
+				CachedAudioTrack->AddSection(*NewAudioSection);
 
-		if ((Sources && Sources->GetSettings().bSaveRecordedAssets) || GEditor == nullptr)
-		{
-			TakesUtils::SaveAsset(RecordedAudio);
+				if ((Sources && Sources->GetSettings().bSaveRecordedAssets) || GEditor == nullptr)
+				{
+					TakesUtils::SaveAsset(RecordedSoundWave);
+				}
+			}
 		}
 	}
 
 	// Reset our audio track pointer
 	CachedAudioTrack = nullptr;
+	RecordedSoundWaves.Empty();
+	
+	ObjectTools::ForceDeleteObjects(AssetsToCleanUp, false);
+	
+	return TArray<UTakeRecorderSource*>();
 }
 
 FText UTakeRecorderMicrophoneAudioSource::GetDisplayTextImpl() const
diff --git a/Engine/Plugins/VirtualProduction/Takes/Source/TakeRecorderSources/Private/TakeRecorderMicrophoneAudioSource.h b/Engine/Plugins/VirtualProduction/Takes/Source/TakeRecorderSources/Private/TakeRecorderMicrophoneAudioSource.h
--- a/Engine/Plugins/VirtualProduction/Takes/Source/TakeRecorderSources/Private/TakeRecorderMicrophoneAudioSource.h
+++ b/Engine/Plugins/VirtualProduction/Takes/Source/TakeRecorderSources/Private/TakeRecorderMicrophoneAudioSource.h
@@ -62,6 +62,7 @@ private:
 	virtual TArray<UTakeRecorderSource*> PreRecording(ULevelSequence* InSequence, FMovieSceneSequenceID InSequenceID, ULevelSequence* InMasterSequence, FManifestSerializer* InManifestSerializer) override;
 	virtual void StartRecording(const FTimecode& InSectionStartTimecode, const FFrameNumber& InSectionFirstFrame, class ULevelSequence* InSequence) override;
 	virtual void StopRecording(class ULevelSequence* InSequence) override;
+	virtual TArray<UTakeRecorderSource*> PostRecording(ULevelSequence* InSequence, class ULevelSequence* InMasterSequence, const bool bCancelled) override;
 	virtual FText GetDisplayTextImpl() const override;
 	virtual void AddContentsToFolder(class UMovieSceneFolder* InFolder) override;
 	virtual bool CanAddSource(UTakeRecorderSources* InSources) const override;
@@ -70,6 +71,8 @@ private:
 private:
 	TWeakObjectPtr<class UMovieSceneAudioTrack> CachedAudioTrack;
 
+	TArray<TWeakObjectPtr<USoundWave>> RecordedSoundWaves;
+
 	TUniquePtr<ISequenceAudioRecorder> AudioRecorder;
 	
 	//Created in PreRecord but used in StartRecording.
diff --git a/Engine/Plugins/VirtualProduction/Takes/Source/TakeRecorderSources/Private/TakeRecorderNearbySpawnedActorSource.cpp b/Engine/Plugins/VirtualProduction/Takes/Source/TakeRecorderSources/Private/TakeRecorderNearbySpawnedActorSource.cpp
--- a/Engine/Plugins/VirtualProduction/Takes/Source/TakeRecorderSources/Private/TakeRecorderNearbySpawnedActorSource.cpp
+++ b/Engine/Plugins/VirtualProduction/Takes/Source/TakeRecorderSources/Private/TakeRecorderNearbySpawnedActorSource.cpp
@@ -44,7 +44,7 @@ TArray<UTakeRecorderSource*> UTakeRecorderNearbySpawnedActorSource::PreRecording
 	return TArray<UTakeRecorderSource*>();
 }
 
-TArray<UTakeRecorderSource*> UTakeRecorderNearbySpawnedActorSource::PostRecording(class ULevelSequence* InSequence, class ULevelSequence* InMasterSequence)
+TArray<UTakeRecorderSource*> UTakeRecorderNearbySpawnedActorSource::PostRecording(class ULevelSequence* InSequence, class ULevelSequence* InMasterSequence, const bool bCancelled)
 {
 	// Remove spawn delegates
 	for (auto It = ActorSpawningDelegateHandles.CreateConstIterator(); It; ++It)
diff --git a/Engine/Plugins/VirtualProduction/Takes/Source/TakeRecorderSources/Private/TakeRecorderNearbySpawnedActorSource.h b/Engine/Plugins/VirtualProduction/Takes/Source/TakeRecorderSources/Private/TakeRecorderNearbySpawnedActorSource.h
--- a/Engine/Plugins/VirtualProduction/Takes/Source/TakeRecorderSources/Private/TakeRecorderNearbySpawnedActorSource.h
+++ b/Engine/Plugins/VirtualProduction/Takes/Source/TakeRecorderSources/Private/TakeRecorderNearbySpawnedActorSource.h
@@ -36,7 +36,7 @@ private:
 
 	// UTakeRecorderSource
 	virtual TArray<UTakeRecorderSource*> PreRecording(ULevelSequence* InSequence, FMovieSceneSequenceID InSequenceID, ULevelSequence* InMasterSequence, FManifestSerializer* InManifestSerializer) override;
-	virtual TArray<UTakeRecorderSource*> PostRecording(class ULevelSequence* InSequence, class ULevelSequence* InMasterSequence) override;
+	virtual TArray<UTakeRecorderSource*> PostRecording(class ULevelSequence* InSequence, class ULevelSequence* InMasterSequence, const bool bCancelled) override;
 	virtual FText GetDisplayTextImpl() const override;
 	virtual FText GetDescriptionTextImpl() const override;
 
diff --git a/Engine/Plugins/VirtualProduction/Takes/Source/TakeRecorderSources/Private/TakeRecorderPlayerSource.cpp b/Engine/Plugins/VirtualProduction/Takes/Source/TakeRecorderSources/Private/TakeRecorderPlayerSource.cpp
--- a/Engine/Plugins/VirtualProduction/Takes/Source/TakeRecorderSources/Private/TakeRecorderPlayerSource.cpp
+++ b/Engine/Plugins/VirtualProduction/Takes/Source/TakeRecorderSources/Private/TakeRecorderPlayerSource.cpp
@@ -69,7 +69,7 @@ TArray<UTakeRecorderSource*> UTakeRecorderPlayerSource::PreRecording(ULevelSeque
 	return NewSources;
 }
 
-TArray<UTakeRecorderSource*> UTakeRecorderPlayerSource::PostRecording(class ULevelSequence* InSequence, class ULevelSequence* InMasterSequence)
+TArray<UTakeRecorderSource*> UTakeRecorderPlayerSource::PostRecording(class ULevelSequence* InSequence, class ULevelSequence* InMasterSequence, const bool bCancelled)
 {
 	TArray<UTakeRecorderSource*> SourcesToRemove;
 
diff --git a/Engine/Plugins/VirtualProduction/Takes/Source/TakeRecorderSources/Private/TakeRecorderPlayerSource.h b/Engine/Plugins/VirtualProduction/Takes/Source/TakeRecorderSources/Private/TakeRecorderPlayerSource.h
--- a/Engine/Plugins/VirtualProduction/Takes/Source/TakeRecorderSources/Private/TakeRecorderPlayerSource.h
+++ b/Engine/Plugins/VirtualProduction/Takes/Source/TakeRecorderSources/Private/TakeRecorderPlayerSource.h
@@ -22,7 +22,7 @@ private:
 
 	// UTakeRecorderSource
 	virtual TArray<UTakeRecorderSource*> PreRecording(ULevelSequence* InSequence, FMovieSceneSequenceID InSequenceID, ULevelSequence* InMasterSequence, FManifestSerializer* InManifestSerializer) override;
-	virtual TArray<UTakeRecorderSource*> PostRecording(class ULevelSequence* InSequence, class ULevelSequence* InMasterSequence) override;
+	virtual TArray<UTakeRecorderSource*> PostRecording(class ULevelSequence* InSequence, class ULevelSequence* InMasterSequence, const bool bCancelled) override;
 	virtual FText GetDisplayTextImpl() const override;
 	virtual bool CanAddSource(UTakeRecorderSources* InSources) const override;
 
diff --git a/Engine/Plugins/VirtualProduction/Takes/Source/TakeRecorderSources/Private/TakeRecorderSourcesModule.cpp b/Engine/Plugins/VirtualProduction/Takes/Source/TakeRecorderSources/Private/TakeRecorderSourcesModule.cpp
--- a/Engine/Plugins/VirtualProduction/Takes/Source/TakeRecorderSources/Private/TakeRecorderSourcesModule.cpp
+++ b/Engine/Plugins/VirtualProduction/Takes/Source/TakeRecorderSources/Private/TakeRecorderSourcesModule.cpp
@@ -732,6 +732,25 @@ public:
 #endif
 	}
 
+	bool HandleCancelRecordTakeCommand(UWorld* InWorld, const TCHAR* InStr, FOutputDevice& Ar)
+	{
+#if WITH_EDITOR
+		if (UTakeRecorder* ActiveRecorder = UTakeRecorder::GetActiveRecorder())
+		{
+			ULevelSequence* ActiveSequence = ActiveRecorder->GetSequence();
+			if (ActiveSequence)
+			{
+				GEditor->GetEditorSubsystem<UAssetEditorSubsystem>()->CloseAllEditorsForAsset(ActiveSequence);
+			}
+
+			ActiveRecorder->Cancel();
+		}
+		return true;
+#else
+		return false;
+#endif
+	}
+
 	virtual bool Exec(UWorld* InWorld, const TCHAR* Cmd, FOutputDevice& Ar) override
 	{
 #if WITH_EDITOR
@@ -743,6 +762,10 @@ public:
 		{
 			return HandleStopRecordTakeCommand(InWorld, Cmd, Ar);
 		}
+		else if (FParse::Command(&Cmd, TEXT("CancelRecordingTake")))
+		{
+			return HandleCancelRecordTakeCommand(InWorld, Cmd, Ar);
+		}
 #endif
 		return false;
 	}
diff --git a/Engine/Plugins/VirtualProduction/Takes/Source/TakeRecorderSources/Private/TakeRecorderWorldSource.cpp b/Engine/Plugins/VirtualProduction/Takes/Source/TakeRecorderSources/Private/TakeRecorderWorldSource.cpp
--- a/Engine/Plugins/VirtualProduction/Takes/Source/TakeRecorderSources/Private/TakeRecorderWorldSource.cpp
+++ b/Engine/Plugins/VirtualProduction/Takes/Source/TakeRecorderSources/Private/TakeRecorderWorldSource.cpp
@@ -98,7 +98,7 @@ TArray<UTakeRecorderSource*> UTakeRecorderWorldSource::PreRecording(ULevelSequen
 	return NewSources;
 }
 
-TArray<UTakeRecorderSource*> UTakeRecorderWorldSource::PostRecording(class ULevelSequence* InSequence, class ULevelSequence* InMasterSequence)
+TArray<UTakeRecorderSource*> UTakeRecorderWorldSource::PostRecording(class ULevelSequence* InSequence, class ULevelSequence* InMasterSequence, const bool bCancelled)
 {
 	TArray<UTakeRecorderSource*> SourcesToRemove;
 
diff --git a/Engine/Plugins/VirtualProduction/Takes/Source/TakeRecorderSources/Private/TakeRecorderWorldSource.h b/Engine/Plugins/VirtualProduction/Takes/Source/TakeRecorderSources/Private/TakeRecorderWorldSource.h
--- a/Engine/Plugins/VirtualProduction/Takes/Source/TakeRecorderSources/Private/TakeRecorderWorldSource.h
+++ b/Engine/Plugins/VirtualProduction/Takes/Source/TakeRecorderSources/Private/TakeRecorderWorldSource.h
@@ -46,7 +46,7 @@ private:
 
 	// UTakeRecorderSource
 	virtual TArray<UTakeRecorderSource*> PreRecording(ULevelSequence* InSequence, FMovieSceneSequenceID InSequenceID, ULevelSequence* InMasterSequence, FManifestSerializer* InManifestSerializer) override;
-	virtual TArray<UTakeRecorderSource*> PostRecording(class ULevelSequence* InSequence, class ULevelSequence* InMasterSequence) override;
+	virtual TArray<UTakeRecorderSource*> PostRecording(class ULevelSequence* InSequence, class ULevelSequence* InMasterSequence, const bool bCancelled) override;
 	virtual bool SupportsTakeNumber() const override { return false; }
 	virtual FText GetDisplayTextImpl() const override;
 	virtual bool CanAddSource(UTakeRecorderSources* InSources) const override;
diff --git a/Engine/Plugins/VirtualProduction/Takes/Source/TakeTrackRecorders/Private/MovieSceneAnimationTrackRecorder.cpp b/Engine/Plugins/VirtualProduction/Takes/Source/TakeTrackRecorders/Private/MovieSceneAnimationTrackRecorder.cpp
--- a/Engine/Plugins/VirtualProduction/Takes/Source/TakeTrackRecorders/Private/MovieSceneAnimationTrackRecorder.cpp
+++ b/Engine/Plugins/VirtualProduction/Takes/Source/TakeTrackRecorders/Private/MovieSceneAnimationTrackRecorder.cpp
@@ -15,6 +15,7 @@
 #include "Engine/Engine.h"
 #include "LevelSequence.h"
 #include "UObject/UObjectBaseUtility.h"
+#include "ObjectTools.h"
 
 #define LOCTEXT_NAMESPACE "MovieSceneAnimationTrackRecorder"
 
@@ -218,6 +219,13 @@ void UMovieSceneAnimationTrackRecorder::FinalizeTrackImpl()
 	}
 }
 
+void UMovieSceneAnimationTrackRecorder::CancelTrackImpl()
+{
+	TArray<UObject*> AssetsToCleanUp;
+	AssetsToCleanUp.Add(GetAnimSequence());
+	ObjectTools::ForceDeleteObjects(AssetsToCleanUp, false);
+}
+
 void UMovieSceneAnimationTrackRecorder::RecordSampleImpl(const FQualifiedFrameTime& CurrentTime)
 {
 	// The animation recorder does most of the work here
diff --git a/Engine/Plugins/VirtualProduction/Takes/Source/TakeTrackRecorders/Public/TrackRecorders/MovieSceneAnimationTrackRecorder.h b/Engine/Plugins/VirtualProduction/Takes/Source/TakeTrackRecorders/Public/TrackRecorders/MovieSceneAnimationTrackRecorder.h
--- a/Engine/Plugins/VirtualProduction/Takes/Source/TakeTrackRecorders/Public/TrackRecorders/MovieSceneAnimationTrackRecorder.h
+++ b/Engine/Plugins/VirtualProduction/Takes/Source/TakeTrackRecorders/Public/TrackRecorders/MovieSceneAnimationTrackRecorder.h
@@ -45,6 +45,7 @@ protected:
 	// UMovieSceneTrackRecorder Interface
 	virtual void CreateTrackImpl() override;
 	virtual void FinalizeTrackImpl() override;
+	virtual void CancelTrackImpl() override;
 	virtual void RecordSampleImpl(const FQualifiedFrameTime& CurrentTime) override;
 	virtual void StopRecordingImpl() override;
 	virtual void SetSavedRecordingDirectory(const FString& InDirectory) override
diff --git a/Engine/Plugins/VirtualProduction/Takes/Source/TakeTrackRecorders/Public/TrackRecorders/MovieSceneTrackRecorder.h b/Engine/Plugins/VirtualProduction/Takes/Source/TakeTrackRecorders/Public/TrackRecorders/MovieSceneTrackRecorder.h
--- a/Engine/Plugins/VirtualProduction/Takes/Source/TakeTrackRecorders/Public/TrackRecorders/MovieSceneTrackRecorder.h
+++ b/Engine/Plugins/VirtualProduction/Takes/Source/TakeTrackRecorders/Public/TrackRecorders/MovieSceneTrackRecorder.h
@@ -5,6 +5,7 @@
 #include "CoreMinimal.h"
 #include "MovieScene.h"
 #include "MovieSceneSection.h"
+#include "MovieSceneTrack.h"
 #include "IMovieSceneTrackRecorderHost.h"
 #include "MovieSceneTrackRecorderSettings.h"
 #include "MovieSceneTrackRecorder.generated.h"
@@ -89,6 +90,23 @@ public:
 		FinalizeTrackImpl();
 	}
 
+	/**
+	* This is called after recording has finished for each track. This allows a track recorder to do any post-processing it may
+	* require such as removing any sections that did not have any changes in them.
+	*/
+	void CancelTrack()
+	{
+		CancelTrackImpl();
+
+		if (UMovieSceneSection* Section = GetMovieSceneSection())
+		{
+			if (UMovieSceneTrack* Track = Cast<UMovieSceneTrack>(Section->GetOuter()))
+			{
+				Track->RemoveSection(*Section);
+			}
+		}
+	}
+
 	/**
 	* This is called each frame and specifies the qualified time that the sampled data should be recorded at. This is passed as 
 	* a FQualifiedFrameTime for better handling of mixed resolution sequences as a user may have modified a sub-sequence to be
@@ -145,6 +163,7 @@ protected:
 	virtual void RecordSampleImpl(const FQualifiedFrameTime& CurrentTime) {}
 	virtual void StopRecordingImpl() {}
 	virtual void FinalizeTrackImpl() {}
+	virtual void CancelTrackImpl() {}
 	virtual UMovieSceneSection* GetMovieSceneSection() const { return nullptr; }
 
 protected:
diff --git a/Engine/Plugins/VirtualProduction/Takes/Source/TakeTrackRecorders/TakeTrackRecorders.Build.cs b/Engine/Plugins/VirtualProduction/Takes/Source/TakeTrackRecorders/TakeTrackRecorders.Build.cs
--- a/Engine/Plugins/VirtualProduction/Takes/Source/TakeTrackRecorders/TakeTrackRecorders.Build.cs
+++ b/Engine/Plugins/VirtualProduction/Takes/Source/TakeTrackRecorders/TakeTrackRecorders.Build.cs
@@ -17,6 +17,7 @@ public class TakeTrackRecorders : ModuleRules
 				"TimeManagement",
 				"SequenceRecorder", // For access to the individual UProperty type recorders
 				"SerializedRecorderInterface",
+				"UnrealEd",
             }
         );
 
diff --git a/Engine/Plugins/VirtualProduction/Takes/Source/TakesCore/Private/TakeRecorderSources.cpp b/Engine/Plugins/VirtualProduction/Takes/Source/TakesCore/Private/TakeRecorderSources.cpp
--- a/Engine/Plugins/VirtualProduction/Takes/Source/TakesCore/Private/TakeRecorderSources.cpp
+++ b/Engine/Plugins/VirtualProduction/Takes/Source/TakesCore/Private/TakeRecorderSources.cpp
@@ -383,7 +383,7 @@ FFrameTime UTakeRecorderSources::TickRecording(class ULevelSequence* InSequence,
 	return CurrentFrameTime.ConvertTo(TargetLevelSequenceTickResolution);
 }
 
-void UTakeRecorderSources::StopRecording(class ULevelSequence* InSequence)
+void UTakeRecorderSources::StopRecording(class ULevelSequence* InSequence, const bool bCancelled)
 {
 	UE_LOG(LogTakesCore, Log, TEXT("StopRecording"));
 
@@ -402,7 +402,7 @@ void UTakeRecorderSources::StopRecording(class ULevelSequence* InSequence)
 	{
 		if (Source->bEnabled)
 		{
-			for (auto SourceToRemove : Source->PostRecording(SourceSubSequenceMap[Source], InSequence))
+			for (auto SourceToRemove : Source->PostRecording(SourceSubSequenceMap[Source], InSequence, bCancelled))
 			{
 				SourcesToRemove.Add(SourceToRemove);
 			}
@@ -431,41 +431,51 @@ void UTakeRecorderSources::StopRecording(class ULevelSequence* InSequence)
 	// as we go (to show the 'animation' of it recording), but we need to restore it to the full length.
 	for (UMovieSceneSubSection* SubSection : ActiveSubSections)
 	{
-		UMovieSceneSequence* SubSequence = SubSection->GetSequence();
-		if (SubSequence)
+		if (bCancelled)
 		{
-			const bool bUpperBoundOnly = false; // Expand the Play Range of the sub-section to encompass all sections within it.
-			TakesUtils::ClampPlaybackRangeToEncompassAllSections(SubSequence->GetMovieScene(), bUpperBoundOnly);
-
-			// Lock the sequence so that it can't be changed without implicitly unlocking it now
-			if (Settings.bAutoLock)
+			if (UMovieSceneTrack* SubTrack = Cast<UMovieSceneTrack>(SubSection->GetOuter()))
 			{
-				SubSequence->GetMovieScene()->SetReadOnly(true);
+				SubTrack->RemoveSection(*SubSection);
 			}
+		}
+		else
+		{
+			UMovieSceneSequence* SubSequence = SubSection->GetSequence();
+			if (SubSequence)
+			{
+				const bool bUpperBoundOnly = false; // Expand the Play Range of the sub-section to encompass all sections within it.
+				TakesUtils::ClampPlaybackRangeToEncompassAllSections(SubSequence->GetMovieScene(), bUpperBoundOnly);
 
-			// Lock the meta data so it can't be changed without implicitly unlocking it now
-			ULevelSequence* SequenceAsset = CastChecked<ULevelSequence>(SubSequence);
-			UTakeMetaData* AssetMetaData = SequenceAsset->FindMetaData<UTakeMetaData>();
-			check(AssetMetaData);
-			AssetMetaData->Lock();
+				// Lock the sequence so that it can't be changed without implicitly unlocking it now
+				if (Settings.bAutoLock)
+				{
+					SubSequence->GetMovieScene()->SetReadOnly(true);
+				}
 
-			SubSection->SetRange(SubSequence->GetMovieScene()->GetPlaybackRange());
+				// Lock the meta data so it can't be changed without implicitly unlocking it now
+				ULevelSequence* SequenceAsset = CastChecked<ULevelSequence>(SubSequence);
+				UTakeMetaData* AssetMetaData = SequenceAsset->FindMetaData<UTakeMetaData>();
+				check(AssetMetaData);
+				AssetMetaData->Lock();
 
-			for (UMovieSceneSection* Section : SubSequence->GetMovieScene()->GetAllSections())
-			{
-				Section->MarkAsChanged();
+				SubSection->SetRange(SubSequence->GetMovieScene()->GetPlaybackRange());
+
+				for (UMovieSceneSection* Section : SubSequence->GetMovieScene()->GetAllSections())
+				{
+					Section->MarkAsChanged();
+				}
+
+				// Re-enable transactional after recording
+				SubSequence->GetMovieScene()->SetFlags(RF_Transactional);
 			}
 
 			// Re-enable transactional after recording
-			SubSequence->GetMovieScene()->SetFlags(RF_Transactional);
-		}
-
-		// Re-enable transactional after recording
-		SubSection->SetFlags(RF_Transactional);
+			SubSection->SetFlags(RF_Transactional);
 
-		if (UMovieSceneTrack* SubTrack = Cast<UMovieSceneTrack>(SubSection->GetOuter()))
-		{
-			SubTrack->SetFlags(RF_Transactional);
+			if (UMovieSceneTrack* SubTrack = Cast<UMovieSceneTrack>(SubSection->GetOuter()))
+			{
+				SubTrack->SetFlags(RF_Transactional);
+			}
 		}
 	}
 
@@ -482,11 +492,19 @@ void UTakeRecorderSources::StopRecording(class ULevelSequence* InSequence)
 		}
 	}
 
-	if (Settings.bSaveRecordedAssets)
+	TArray<UObject*> AssetsToCleanUp;
+	for (auto SourceSubSequence : SourceSubSequenceMap)
 	{
-		for (auto SourceSubSequence : SourceSubSequenceMap)
+		if (SourceSubSequence.Value)
 		{
-			if (SourceSubSequence.Value)
+			if (bCancelled)
+			{
+				if (SourceSubSequence.Value.Get() != InSequence)
+				{
+					AssetsToCleanUp.Add(SourceSubSequence.Value.Get());
+				}
+			}
+			else if (Settings.bSaveRecordedAssets)
 			{
 				TakesUtils::SaveAsset(SourceSubSequence.Value);
 			}
@@ -498,6 +516,7 @@ void UTakeRecorderSources::StopRecording(class ULevelSequence* InSequence)
 	CreatedManifestSerializers.Empty();
 	CachedManifestSerializer = nullptr;
 	CachedLevelSequence = nullptr;
+	ObjectTools::ForceDeleteObjects(AssetsToCleanUp, false);
 }
 
 ULevelSequence* UTakeRecorderSources::CreateSubSequenceForSource(ULevelSequence* InMasterSequence, const FString& SubSequenceTrackName, const FString& SubSequenceAssetName)
diff --git a/Engine/Plugins/VirtualProduction/Takes/Source/TakesCore/Public/TakeRecorderSource.h b/Engine/Plugins/VirtualProduction/Takes/Source/TakesCore/Public/TakeRecorderSource.h
--- a/Engine/Plugins/VirtualProduction/Takes/Source/TakesCore/Public/TakeRecorderSource.h
+++ b/Engine/Plugins/VirtualProduction/Takes/Source/TakesCore/Public/TakeRecorderSource.h
@@ -117,7 +117,7 @@ public:
 	* @return An array of take recorder sources to be removed. Likely to match the same list as PreRecording. Can be an empty list
 	* if no additional sources are required by this source.
 	*/
-	virtual TArray<UTakeRecorderSource*> PostRecording(class ULevelSequence* InSequence, class ULevelSequence* InMasterSequence) { return TArray<UTakeRecorderSource*>(); }
+	virtual TArray<UTakeRecorderSource*> PostRecording(class ULevelSequence* InSequence, class ULevelSequence* InMasterSequence, const bool bCancelled) { return TArray<UTakeRecorderSource*>(); }
 
 
 	/**
diff --git a/Engine/Plugins/VirtualProduction/Takes/Source/TakesCore/Public/TakeRecorderSources.h b/Engine/Plugins/VirtualProduction/Takes/Source/TakesCore/Public/TakeRecorderSources.h
--- a/Engine/Plugins/VirtualProduction/Takes/Source/TakesCore/Public/TakeRecorderSources.h
+++ b/Engine/Plugins/VirtualProduction/Takes/Source/TakesCore/Public/TakeRecorderSources.h
@@ -158,7 +158,7 @@ public:
 	* Stop recording pass
 	*
 	*/
-	void StopRecording(class ULevelSequence* InSequence);
+	void StopRecording(class ULevelSequence* InSequence, const bool bCancelled);
 
 public:
 	/*
