Take Recorder: Check GEditor before deleting assets

diff --git a/Engine/Plugins/VirtualProduction/Takes/Source/TakeRecorderSources/Private/TakeRecorderMicrophoneAudioSource.cpp b/Engine/Plugins/VirtualProduction/Takes/Source/TakeRecorderSources/Private/TakeRecorderMicrophoneAudioSource.cpp
--- a/Engine/Plugins/VirtualProduction/Takes/Source/TakeRecorderSources/Private/TakeRecorderMicrophoneAudioSource.cpp
+++ b/Engine/Plugins/VirtualProduction/Takes/Source/TakeRecorderSources/Private/TakeRecorderMicrophoneAudioSource.cpp
@@ -252,8 +252,11 @@ TArray<UTakeRecorderSource*> UTakeRecorderMicrophoneAudioSource::PostRecording(U
 	CachedAudioTrack = nullptr;
 	RecordedSoundWaves.Empty();
 	
-	ObjectTools::ForceDeleteObjects(AssetsToCleanUp, false);
-	
+	if (GEditor && AssetsToCleanUp.Num() > 0)
+	{
+		ObjectTools::ForceDeleteObjects(AssetsToCleanUp, false);
+	}
+
 	return TArray<UTakeRecorderSource*>();
 }
 
diff --git a/Engine/Plugins/VirtualProduction/Takes/Source/TakeTrackRecorders/Private/MovieSceneAnimationTrackRecorder.cpp b/Engine/Plugins/VirtualProduction/Takes/Source/TakeTrackRecorders/Private/MovieSceneAnimationTrackRecorder.cpp
--- a/Engine/Plugins/VirtualProduction/Takes/Source/TakeTrackRecorders/Private/MovieSceneAnimationTrackRecorder.cpp
+++ b/Engine/Plugins/VirtualProduction/Takes/Source/TakeTrackRecorders/Private/MovieSceneAnimationTrackRecorder.cpp
@@ -224,7 +224,11 @@ void UMovieSceneAnimationTrackRecorder::CancelTrackImpl()
 {
 	TArray<UObject*> AssetsToCleanUp;
 	AssetsToCleanUp.Add(GetAnimSequence());
-	ObjectTools::ForceDeleteObjects(AssetsToCleanUp, false);
+	
+	if (GEditor && AssetsToCleanUp.Num() > 0)
+	{
+		ObjectTools::ForceDeleteObjects(AssetsToCleanUp, false);
+	}
 }
 
 void UMovieSceneAnimationTrackRecorder::RecordSampleImpl(const FQualifiedFrameTime& CurrentTime)
diff --git a/Engine/Plugins/VirtualProduction/Takes/Source/TakesCore/Private/TakeRecorderSources.cpp b/Engine/Plugins/VirtualProduction/Takes/Source/TakesCore/Private/TakeRecorderSources.cpp
--- a/Engine/Plugins/VirtualProduction/Takes/Source/TakesCore/Private/TakeRecorderSources.cpp
+++ b/Engine/Plugins/VirtualProduction/Takes/Source/TakesCore/Private/TakeRecorderSources.cpp
@@ -516,7 +516,11 @@ void UTakeRecorderSources::StopRecording(class ULevelSequence* InSequence, const
 	CreatedManifestSerializers.Empty();
 	CachedManifestSerializer = nullptr;
 	CachedLevelSequence = nullptr;
-	ObjectTools::ForceDeleteObjects(AssetsToCleanUp, false);
+
+	if (GEditor && AssetsToCleanUp.Num() > 0)
+	{
+		ObjectTools::ForceDeleteObjects(AssetsToCleanUp, false);
+	}
 }
 
 ULevelSequence* UTakeRecorderSources::CreateSubSequenceForSource(ULevelSequence* InMasterSequence, const FString& SubSequenceTrackName, const FString& SubSequenceAssetName)
