Sequencer: Add include for geditor

diff --git a/Engine/Plugins/VirtualProduction/Takes/Source/TakeTrackRecorders/Private/MovieSceneAnimationTrackRecorder.cpp b/Engine/Plugins/VirtualProduction/Takes/Source/TakeTrackRecorders/Private/MovieSceneAnimationTrackRecorder.cpp
--- a/Engine/Plugins/VirtualProduction/Takes/Source/TakeTrackRecorders/Private/MovieSceneAnimationTrackRecorder.cpp
+++ b/Engine/Plugins/VirtualProduction/Takes/Source/TakeTrackRecorders/Private/MovieSceneAnimationTrackRecorder.cpp
@@ -17,6 +17,7 @@
 #include "UObject/SavePackage.h"
 #include "UObject/UObjectBaseUtility.h"
 #include "ObjectTools.h"
+#include "Editor.h"
 
 #define LOCTEXT_NAMESPACE "MovieSceneAnimationTrackRecorder"
 
