Take Recorder: Added FinalizeRecording() to UTakeRecorderSource which allows sources to do a final clean up which has no dependencies on other sources. TakeRecorderActorSource needs to null out the TargetLevelSequence and the MasterLevelSequence when done recording but can't do so in PostRecording since other sources might depend on it when constructing GetLevelSequenceID(). So, this needs to be done after PostRecording, otherwise there will be hard references left over to the take recorder sources.

diff --git a/Engine/Plugins/VirtualProduction/Takes/Source/TakeRecorderSources/Private/TakeRecorderActorSource.cpp b/Engine/Plugins/VirtualProduction/Takes/Source/TakeRecorderSources/Private/TakeRecorderActorSource.cpp
--- a/Engine/Plugins/VirtualProduction/Takes/Source/TakeRecorderSources/Private/TakeRecorderActorSource.cpp
+++ b/Engine/Plugins/VirtualProduction/Takes/Source/TakeRecorderSources/Private/TakeRecorderActorSource.cpp
@@ -816,14 +816,17 @@ TArray<UTakeRecorderSource*> UTakeRecorderActorSource::PostRecording(ULevelSeque
 	CachedObjectTemplate = nullptr;
 	CachedComponentList.Empty();
 
-	//DON"T null these out they can be used for doing cross sequence object binding via GetLevelSequenceID
-	//TargetLevelSequence = nullptr;
-	//MasterLevelSequence = nullptr;
-
 	// We may have generated some temporary recording sources
 	return AddedActorSources;
 }
 
+void UTakeRecorderActorSource::FinalizeRecording()
+{
+	// Null these out there and NOT in PostRecording because they are used for cross sequence object binding via GetLevelSequenceID in PostRecording
+	TargetLevelSequence = nullptr;
+	MasterLevelSequence = nullptr;
+}
+
 void UTakeRecorderActorSource::PostProcessTrackRecorders(ULevelSequence* InSequence)
 {
 	FTakeRecorderParameters Parameters;
diff --git a/Engine/Plugins/VirtualProduction/Takes/Source/TakeRecorderSources/Private/TakeRecorderActorSource.h b/Engine/Plugins/VirtualProduction/Takes/Source/TakeRecorderSources/Private/TakeRecorderActorSource.h
--- a/Engine/Plugins/VirtualProduction/Takes/Source/TakeRecorderSources/Private/TakeRecorderActorSource.h
+++ b/Engine/Plugins/VirtualProduction/Takes/Source/TakeRecorderSources/Private/TakeRecorderActorSource.h
@@ -125,6 +125,7 @@ public:
 	virtual void TickRecording(const FQualifiedFrameTime& CurrentSequenceTime) override;
 	virtual void StopRecording(class ULevelSequence* InSequence) override;
 	virtual TArray<UTakeRecorderSource*> PostRecording(class ULevelSequence* InSequence, class ULevelSequence* InMasterSequence, const bool bCancelled) override;
+	virtual void FinalizeRecording() override;
 	virtual TArray<UObject*> GetAdditionalSettingsObjects() const { return TArray<UObject*>(FactorySettings); }
 	virtual FString GetSubsceneTrackName(ULevelSequence* InSequence) const override;
 	virtual FString GetSubsceneAssetName(ULevelSequence* InSequence) const override;
diff --git a/Engine/Plugins/VirtualProduction/Takes/Source/TakesCore/Private/TakeRecorderSources.cpp b/Engine/Plugins/VirtualProduction/Takes/Source/TakesCore/Private/TakeRecorderSources.cpp
--- a/Engine/Plugins/VirtualProduction/Takes/Source/TakesCore/Private/TakeRecorderSources.cpp
+++ b/Engine/Plugins/VirtualProduction/Takes/Source/TakesCore/Private/TakeRecorderSources.cpp
@@ -418,6 +418,11 @@ void UTakeRecorderSources::StopRecording(class ULevelSequence* InSequence, const
 		++SourcesSerialNumber;
 	}
 
+	for (UTakeRecorderSource* Source : Sources)
+	{
+		Source->FinalizeRecording();
+	}
+
 	// Re-enable transactional after recording
 	InSequence->GetMovieScene()->SetFlags(RF_Transactional);
 	for (UMovieSceneSection* Section : InSequence->GetMovieScene()->GetAllSections())
diff --git a/Engine/Plugins/VirtualProduction/Takes/Source/TakesCore/Public/TakeRecorderSource.h b/Engine/Plugins/VirtualProduction/Takes/Source/TakesCore/Public/TakeRecorderSource.h
--- a/Engine/Plugins/VirtualProduction/Takes/Source/TakesCore/Public/TakeRecorderSource.h
+++ b/Engine/Plugins/VirtualProduction/Takes/Source/TakesCore/Public/TakeRecorderSource.h
@@ -119,6 +119,10 @@ public:
 	*/
 	virtual TArray<UTakeRecorderSource*> PostRecording(class ULevelSequence* InSequence, class ULevelSequence* InMasterSequence, const bool bCancelled) { return TArray<UTakeRecorderSource*>(); }
 
+	/*
+	 * This is called on all sources after post recording. There should not be any dependencies between sources at this point.
+	 */
+	virtual void FinalizeRecording() {}
 
 	/**
 	* This allows a Source to return an array of dynamically spawned settings objects for that source.
