Take Recorder: Recorder for FLinearColor

diff --git a/Engine/Plugins/VirtualProduction/Takes/Source/TakeTrackRecorders/Private/MovieScenePropertyTrackRecorder.cpp b/Engine/Plugins/VirtualProduction/Takes/Source/TakeTrackRecorders/Private/MovieScenePropertyTrackRecorder.cpp
--- a/Engine/Plugins/VirtualProduction/Takes/Source/TakeTrackRecorders/Private/MovieScenePropertyTrackRecorder.cpp
+++ b/Engine/Plugins/VirtualProduction/Takes/Source/TakeTrackRecorders/Private/MovieScenePropertyTrackRecorder.cpp
@@ -40,6 +40,10 @@ bool FMovieScenePropertyTrackRecorderFactory::CanRecordProperty(UObject* InObjec
 		{
 			return true;
 		}
+		else if (StructProperty->Struct->GetFName() == NAME_LinearColor)
+		{
+			return true;
+		}
 	}
 
 	// We only know how to make generic tracks for the types above
@@ -91,6 +95,9 @@ UMovieSceneTrackRecorder* FMovieScenePropertyTrackRecorderFactory::CreateTrackRe
 	case ESerializedPropertyType::ColorType:
 		TrackRecorder->PropertyRecorder = MakeShareable(new FMovieSceneTrackPropertyRecorder<FColor>(Binding));
 		break;
+	case ESerializedPropertyType::LinearColorType:
+		TrackRecorder->PropertyRecorder = MakeShareable(new FMovieSceneTrackPropertyRecorder<FLinearColor>(Binding));
+		break;
 
 	}
 	return TrackRecorder;
@@ -149,6 +156,10 @@ void UMovieScenePropertyTrackRecorder::CreateTrackImpl()
  			{
 				PropertyRecorder = MakeShareable(new FMovieSceneTrackPropertyRecorder<FColor>(Binding));
  			}
+ 			else if (StructProperty->Struct->GetFName() == NAME_LinearColor)
+ 			{
+				PropertyRecorder = MakeShareable(new FMovieSceneTrackPropertyRecorder<FLinearColor>(Binding));
+ 			}
  		} 
 		ensure(PropertyRecorder);
 		PropertyRecorder->SetSavedRecordingDirectory(Directory);
diff --git a/Engine/Plugins/VirtualProduction/Takes/Source/TakeTrackRecorders/Private/MovieSceneTrackPropertyRecorder.cpp b/Engine/Plugins/VirtualProduction/Takes/Source/TakeTrackRecorders/Private/MovieSceneTrackPropertyRecorder.cpp
--- a/Engine/Plugins/VirtualProduction/Takes/Source/TakeTrackRecorders/Private/MovieSceneTrackPropertyRecorder.cpp
+++ b/Engine/Plugins/VirtualProduction/Takes/Source/TakeTrackRecorders/Private/MovieSceneTrackPropertyRecorder.cpp
@@ -880,6 +880,213 @@ bool FMovieSceneTrackPropertyRecorder<FColor>::LoadRecordedFile(const FString& F
 	return false;
 }
 
+template <>
+bool FMovieSceneTrackPropertyRecorder<FLinearColor>::ShouldAddNewKey(const FLinearColor& InNewValue) const
+{
+	return PreviousValue.R != InNewValue.R  || PreviousValue.G != InNewValue.G || PreviousValue.B != InNewValue.B;
+}
+
+template <>
+UMovieSceneSection* FMovieSceneTrackPropertyRecorder<FLinearColor>::AddSection(const FString& TrackDisplayName, UMovieScene* InMovieScene, const FGuid& InGuid, bool bSetDefault)
+{
+	FName TrackName = *Binding.GetPropertyPath();
+	UMovieSceneColorTrack* Track = InMovieScene->FindTrack<UMovieSceneColorTrack>(InGuid, TrackName);
+	if (!Track)
+	{
+		Track = InMovieScene->AddTrack<UMovieSceneColorTrack>(InGuid);
+	}
+	else
+	{
+		Track->RemoveAllAnimationData();
+	}
+
+	if (Track)
+	{
+		Track->SetPropertyNameAndPath(Binding.GetPropertyName(), Binding.GetPropertyPath());
+
+		UMovieSceneColorSection* Section = Cast<UMovieSceneColorSection>(Track->CreateNewSection());
+
+		// We only set the track defaults when we're not loading from a serialized recording. Serialized recordings don't store channel defaults but will always store a
+			// key on the first frame which will accomplish the same.
+		if (bSetDefault)
+		{
+			TArrayView<FMovieSceneFloatChannel*> FloatChannels = Section->GetChannelProxy().GetChannels<FMovieSceneFloatChannel>();
+		
+			FloatChannels[0]->SetDefault(PreviousValue.R);
+			FloatChannels[1]->SetDefault(PreviousValue.G);
+			FloatChannels[2]->SetDefault(PreviousValue.B);
+			FloatChannels[3]->SetDefault(PreviousValue.A);
+		}
+
+		Track->AddSection(*Section);
+
+		return Section;
+	}
+
+	return nullptr;
+}
+
+template <>
+void FMovieSceneTrackPropertyRecorder<FLinearColor>::PostCreate(IMovieSceneTrackRecorderHost* InRecordingHost, UObject* InObjectToRecord, class UMovieScene* InMovieScene, const FGuid& InGuid, bool bOpenSerializer)
+{
+}
+
+template <>
+void FMovieSceneTrackPropertyRecorder<FLinearColor>::AddKeyToSection(UMovieSceneSection* InSection, const FPropertyKey<FLinearColor>& InKey)
+{
+	TArrayView<FMovieSceneFloatChannel*> FloatChannels = InSection->GetChannelProxy().GetChannels<FMovieSceneFloatChannel>();
+	FloatChannels[0]->AddCubicKey(InKey.Time, InKey.Value.R, RCTM_Break);
+	FloatChannels[1]->AddCubicKey(InKey.Time, InKey.Value.G, RCTM_Break);
+	FloatChannels[2]->AddCubicKey(InKey.Time, InKey.Value.B, RCTM_Break);
+	FloatChannels[3]->AddCubicKey(InKey.Time, InKey.Value.A, RCTM_Break);
+}
+
+template <>
+void FMovieSceneTrackPropertyRecorder<FLinearColor>::ReduceKeys(UMovieSceneSection* InSection, float ReduceKeysTolerance)
+{
+	TArrayView<FMovieSceneFloatChannel*> FloatChannels = InSection->GetChannelProxy().GetChannels<FMovieSceneFloatChannel>();
+
+	FKeyDataOptimizationParams Params;
+	Params.bAutoSetInterpolation = true;
+	Params.Tolerance = ReduceKeysTolerance;
+	UE::MovieScene::Optimize(FloatChannels[0], Params);
+	UE::MovieScene::Optimize(FloatChannels[1], Params);
+	UE::MovieScene::Optimize(FloatChannels[2], Params);
+	UE::MovieScene::Optimize(FloatChannels[3], Params);
+}
+
+template <>
+FLinearColor FMovieSceneTrackPropertyRecorder<FLinearColor>::GetDefaultValue(UMovieSceneSection* InSection)
+{
+	FLinearColor DefaultValue(0.f, 0.f, 0.f, 1.f);
+
+	TArrayView<FMovieSceneFloatChannel*> FloatChannels = InSection->GetChannelProxy().GetChannels<FMovieSceneFloatChannel>();
+
+	if (FloatChannels[0]->GetNumKeys() > 0)
+	{
+		DefaultValue.R = FloatChannels[0]->GetValues()[0].Value;
+	}
+	else if (FloatChannels[0]->GetDefault().IsSet())
+	{
+		DefaultValue.R = FloatChannels[0]->GetDefault().GetValue();
+	}
+
+	if (FloatChannels[1]->GetNumKeys() > 0)
+	{
+		DefaultValue.G = FloatChannels[1]->GetValues()[0].Value;
+	}
+	else if (FloatChannels[1]->GetDefault().IsSet())
+	{
+		DefaultValue.G = FloatChannels[1]->GetDefault().GetValue();
+	}
+
+	if (FloatChannels[2]->GetNumKeys() > 0)
+	{
+		DefaultValue.B = FloatChannels[2]->GetValues()[0].Value;
+	}
+	else if (FloatChannels[2]->GetDefault().IsSet())
+	{
+		DefaultValue.B = FloatChannels[2]->GetDefault().GetValue();
+	}
+
+	if (FloatChannels[3]->GetNumKeys() > 0)
+	{
+		DefaultValue.A = FloatChannels[3]->GetValues()[0].Value;
+	}
+	else if (FloatChannels[3]->GetDefault().IsSet())
+	{
+		DefaultValue.A = FloatChannels[3]->GetDefault().GetValue();
+	}
+
+	return DefaultValue;
+}
+
+template <>
+void FMovieSceneTrackPropertyRecorder<FLinearColor>::SetDefaultValue(UMovieSceneSection* InSection, const FLinearColor& InDefaultValue)
+{
+	TArrayView<FMovieSceneFloatChannel*> FloatChannels = InSection->GetChannelProxy().GetChannels<FMovieSceneFloatChannel>();
+	FloatChannels[0]->SetDefault(InDefaultValue.R);
+	FloatChannels[1]->SetDefault(InDefaultValue.G);
+	FloatChannels[2]->SetDefault(InDefaultValue.B);
+	FloatChannels[3]->SetDefault(InDefaultValue.A);
+}
+
+
+template<>
+bool FMovieSceneTrackPropertyRecorder<FLinearColor>::OpenSerializer(const FString& InObjectName, const FName& InPropertyName, const FString& InTrackDisplayName, const FGuid& InGuid)
+{
+/*
+	FName SerializedType("Property");
+	FFrameRate   TickResolution = MovieSceneSection->GetTypedOuter<UMovieScene>()->GetTickResolution();
+	FPropertyFileHeader Header(TickResolution, SerializedType, InGuid);
+	Header.PropertyName = InPropertyName;
+	Header.TrackDisplayName = InTrackDisplayName;
+	Header.PropertyType = ESerializedPropertyType::ColorType;
+	FText Error;
+	FString FileName = FString::Printf(TEXT("%s_%s_%s"), *(SerializedType.ToString()), *InObjectName, *(InPropertyName.ToString()));
+
+	if (!Serializer.OpenForWrite(FileName, Header, Error))
+	{
+		UE_LOG(PropertySerialization, Warning, TEXT("Error Opening Property File: Object '%s' Property '%s' Error: '%s'"), *InObjectName, *InPropertyName.ToString(), *Error.ToString());
+		return false;
+	}
+	*/
+	return true;
+}
+
+template <>
+bool FMovieSceneTrackPropertyRecorder<FLinearColor>::LoadRecordedFile(const FString& FileName, UMovieScene *InMovieScene, TMap<FGuid, AActor*>& ActorGuidToActorMap, TFunction<void()> InCompletionCallback)
+{
+/*
+	bool bFileExists = Serializer.DoesFileExist(FileName);
+	if (bFileExists)
+	{
+		FText Error;
+		FPropertyFileHeader Header;
+
+		if (Serializer.OpenForRead(FileName, Header, Error))
+		{
+			MovieSceneSection = AddSection(Header.TrackDisplayName, InMovieScene, Header.Guid, false);
+
+			Serializer.GetDataRanges([this, InMovieScene, FileName, Header, InCompletionCallback](uint64 InMinFrameId, uint64 InMaxFrameId)
+			{
+				auto OnReadComplete = [this, InMovieScene, Header, InCompletionCallback]()
+				{
+					TArray<FPropertySerializedColorFrame> &InFrames = Serializer.ResultData;
+					if (InFrames.Num() > 0)
+					{
+						FFrameRate InFrameRate = Header.TickResolution;
+						for (const FPropertySerializedColorFrame& SerializedFrame : InFrames)
+						{
+							const FPropertySerializedColor& Frame = SerializedFrame.Frame;
+							FFrameRate   TickResolution = MovieSceneSection->GetTypedOuter<UMovieScene>()->GetTickResolution();
+							FFrameTime FrameTime = FFrameRate::TransformTime(Frame.Time, InFrameRate, TickResolution);
+							FFrameNumber CurrentFrame = FrameTime.FrameNumber;
+							FPropertyKey<FColor> Key;
+							Key.Time = CurrentFrame;
+							Key.Value = FColor(Frame.Value); 
+							AddKeyToSection(MovieSceneSection.Get(), Key);
+							MovieSceneSection->ExpandToFrame(CurrentFrame);
+						}
+					}
+					Serializer.Close();
+					InCompletionCallback();
+				}; //callback
+
+				Serializer.ReadFramesAtFrameRange(InMinFrameId, InMaxFrameId, OnReadComplete);
+
+			});
+			return true;
+		}
+		else
+		{
+			Serializer.Close();
+		}
+	}
+	*/
+	return false;
+}
+
 template <>
 bool FMovieSceneTrackPropertyRecorder<FVector3f>::ShouldAddNewKey(const FVector3f& InNewValue) const
 {
@@ -1566,6 +1773,7 @@ template class FMovieSceneTrackPropertyRecorder<bool>;
 template class FMovieSceneTrackPropertyRecorder<uint8>;
 template class FMovieSceneTrackPropertyRecorder<float>;
 template class FMovieSceneTrackPropertyRecorder<FColor>;
+template class FMovieSceneTrackPropertyRecorder<FLinearColor>;
 template class FMovieSceneTrackPropertyRecorder<FVector>;
 template class FMovieSceneTrackPropertyRecorder<int32>;
 template class FMovieSceneTrackPropertyRecorder<FString>;
@@ -1574,6 +1782,7 @@ template struct FSerializedProperty<bool>;
 template struct FSerializedProperty<uint8>;
 template struct FSerializedProperty<float>;
 template struct FSerializedProperty<FColor>;
+template struct FSerializedProperty<FLinearColor>;
 template struct FSerializedProperty<FVector>;
 template struct FSerializedProperty<int32>;
 template struct FSerializedProperty<FString>;
diff --git a/Engine/Source/Editor/SerializedRecorderInterface/Private/MovieSceneSectionSerialization.cpp b/Engine/Source/Editor/SerializedRecorderInterface/Private/MovieSceneSectionSerialization.cpp
--- a/Engine/Source/Editor/SerializedRecorderInterface/Private/MovieSceneSectionSerialization.cpp
+++ b/Engine/Source/Editor/SerializedRecorderInterface/Private/MovieSceneSectionSerialization.cpp
@@ -96,6 +96,11 @@ FMovieSceneSerializerRunnable<FPropertyFileHeader, FSerializedProperty<FString>>
 template<> SERIALIZEDRECORDERINTERFACE_API
 FRunnableThread* TMovieSceneSerializer<FPropertyFileHeader, FSerializedProperty<FString>>::Thread = nullptr;
 
+template<> SERIALIZEDRECORDERINTERFACE_API
+FMovieSceneSerializerRunnable<FPropertyFileHeader, FSerializedProperty<FLinearColor>>* TMovieSceneSerializer<FPropertyFileHeader, FSerializedProperty<FLinearColor>>::Runnable = nullptr;
+template<> SERIALIZEDRECORDERINTERFACE_API
+FRunnableThread* TMovieSceneSerializer<FPropertyFileHeader, FSerializedProperty<FLinearColor>>::Thread = nullptr;
+
 /** External Tempplate statics.  Need to define here due to linkage issues*/
 
 /* From TakeRecorderActorSource */
diff --git a/Engine/Source/Editor/SerializedRecorderInterface/Public/Serializers/MovieScenePropertySerialization.h b/Engine/Source/Editor/SerializedRecorderInterface/Public/Serializers/MovieScenePropertySerialization.h
--- a/Engine/Source/Editor/SerializedRecorderInterface/Public/Serializers/MovieScenePropertySerialization.h
+++ b/Engine/Source/Editor/SerializedRecorderInterface/Public/Serializers/MovieScenePropertySerialization.h
@@ -18,7 +18,8 @@ enum class ESerializedPropertyType : uint8 {
 	Vector3dType,
 	ColorType,
 	IntegerType,
-	StringType
+	StringType,
+	LinearColorType
 };
 
 struct FPropertyFileHeader
@@ -103,6 +104,10 @@ struct FPropertyFileHeader
 				{
 					PropertyType = (ESerializedPropertyType::ColorType);
 				}
+				else if (StructProperty->Struct->GetFName() == NAME_LinearColor)
+				{
+					PropertyType = (ESerializedPropertyType::LinearColorType);
+				}
 			}
 			else if (Property->IsA<FIntProperty>())
 			{
@@ -176,4 +181,5 @@ using FPropertySerializedIntegerFrame = TMovieSceneSerializedFrame<FSerializedPr
 using FPropertySerializedString = FSerializedProperty<FString>;
 using FPropertySerializedStringFrame = TMovieSceneSerializedFrame<FSerializedProperty<FString>>;
 
-
+using FPropertySerializedLinearColor = FSerializedProperty<FLinearColor>;
+using FPropertySerializedLinearColorFrame = TMovieSceneSerializedFrame<FSerializedProperty<FLinearColor>>;
\ No newline at end of file
