MovieScene: query movie scenes and sequences for their earliest TimecodeSource instead of using their own property

UMovieScene is a container that may represent multiple pieces of media as
sections, each of which may have their own distinct timecode source. As a result,
the timecode source associated with the UMovieScene itself doesn't correspond to
any particular piece of media, but rather was just intended as the timecode at which
the movie scene was recorded.

These changes update code sites that read the TimecodeSource property directly to
instead query the UMovieScene/UMovieSceneSequence for the earliest timecode source
of any of its sections using GetEarliestTimecodeSource(). The TimecodeSource property
of UMovieScene will ultimately be removed in a subsequent change.

diff --git a/Engine/Plugins/VirtualProduction/Takes/Source/TakeTrackRecorders/Private/MovieScene3DAttachTrackRecorder.cpp b/Engine/Plugins/VirtualProduction/Takes/Source/TakeTrackRecorders/Private/MovieScene3DAttachTrackRecorder.cpp
--- a/Engine/Plugins/VirtualProduction/Takes/Source/TakeTrackRecorders/Private/MovieScene3DAttachTrackRecorder.cpp
+++ b/Engine/Plugins/VirtualProduction/Takes/Source/TakeTrackRecorders/Private/MovieScene3DAttachTrackRecorder.cpp
@@ -55,7 +55,7 @@ void UMovieScene3DAttachTrackRecorder::RecordSampleImpl(const FQualifiedFrameTim
 				MovieSceneSection->AttachSocketName = SocketName;
 				MovieSceneSection->AttachComponentName = ComponentName;
 
-				MovieSceneSection->TimecodeSource = MovieScene->TimecodeSource;
+				MovieSceneSection->TimecodeSource = MovieScene->GetEarliestTimecodeSource();
 				MovieSceneSection->SetRange(TRange<FFrameNumber>(CurrentFrame, CurrentFrame));
 
 				FMovieSceneSequenceID TargetSequenceID = OwningTakeRecorderSource->GetLevelSequenceID(AttachedToActor);
diff --git a/Engine/Plugins/VirtualProduction/Takes/Source/TakeTrackRecorders/Private/MovieScene3DTransformTrackRecorder.cpp b/Engine/Plugins/VirtualProduction/Takes/Source/TakeTrackRecorders/Private/MovieScene3DTransformTrackRecorder.cpp
--- a/Engine/Plugins/VirtualProduction/Takes/Source/TakeTrackRecorders/Private/MovieScene3DTransformTrackRecorder.cpp
+++ b/Engine/Plugins/VirtualProduction/Takes/Source/TakeTrackRecorders/Private/MovieScene3DTransformTrackRecorder.cpp
@@ -624,7 +624,6 @@ bool UMovieScene3DTransformTrackRecorder::LoadRecordedFile(const FString& FileNa
 
 							MovieSceneTrack->AddSection(*MovieSceneSection);
 							TArrayView<FMovieSceneDoubleChannel*> DoubleChannels = MovieSceneSection->GetChannelProxy().GetChannels<FMovieSceneDoubleChannel>();
-							//MovieSceneSection->TimecodeSource = SequenceRecorderUtils::GetTimecodeSource();
 
 							bool InFirst = true;
 							// @todo: sequencer-timecode: this was previously never actually used - should it be??
diff --git a/Engine/Plugins/VirtualProduction/Takes/Source/TakeTrackRecorders/Private/MovieSceneAnimationTrackRecorder.cpp b/Engine/Plugins/VirtualProduction/Takes/Source/TakeTrackRecorders/Private/MovieSceneAnimationTrackRecorder.cpp
--- a/Engine/Plugins/VirtualProduction/Takes/Source/TakeTrackRecorders/Private/MovieSceneAnimationTrackRecorder.cpp
+++ b/Engine/Plugins/VirtualProduction/Takes/Source/TakeTrackRecorders/Private/MovieSceneAnimationTrackRecorder.cpp
@@ -9,6 +9,7 @@
 #include "MovieScene.h"
 #include "AssetRegistryModule.h"
 #include "SequenceRecorderSettings.h"
+#include "SequenceRecorderUtils.h"
 #include "Components/SkeletalMeshComponent.h"
 #include "Animation/AnimSequence.h"
 #include "Engine/TimecodeProvider.h"
@@ -244,7 +245,7 @@ void UMovieSceneAnimationTrackRecorder::RecordSampleImpl(const FQualifiedFrameTi
 		//Reset the start times based upon when the animation really starts.
 		if (MovieSceneSection.IsValid())
 		{
-			MovieSceneSection->TimecodeSource = FMovieSceneTimecodeSource(FApp::GetTimecode());
+			MovieSceneSection->TimecodeSource = SequenceRecorderUtils::GetTimecodeSource();
 			FFrameRate   TickResolution = MovieSceneSection->GetTypedOuter<UMovieScene>()->GetTickResolution();
 			FFrameNumber CurrentFrame = CurrentTime.ConvertTo(TickResolution).FloorToFrame();
 
diff --git a/Engine/Source/Runtime/LevelSequence/Private/LevelSequencePlayer.cpp b/Engine/Source/Runtime/LevelSequence/Private/LevelSequencePlayer.cpp
--- a/Engine/Source/Runtime/LevelSequence/Private/LevelSequencePlayer.cpp
+++ b/Engine/Source/Runtime/LevelSequence/Private/LevelSequencePlayer.cpp
@@ -458,7 +458,7 @@ void ULevelSequencePlayer::TakeFrameSnapshot(FLevelSequencePlayerSnapshot& OutSn
 	UMovieScene* MovieScene = Sequence->GetMovieScene();
 
 #if WITH_EDITORONLY_DATA
-	OutSnapshot.SourceTimecode = MovieScene->TimecodeSource.Timecode.ToString();
+	OutSnapshot.SourceTimecode = MovieScene->GetEarliestTimecodeSource().Timecode.ToString();
 #endif
 
 	UMovieSceneCinematicShotTrack* ShotTrack = MovieScene->FindMasterTrack<UMovieSceneCinematicShotTrack>();
diff --git a/Engine/Source/Runtime/MovieScene/Private/Tracks/MovieSceneSubTrack.cpp b/Engine/Source/Runtime/MovieScene/Private/Tracks/MovieSceneSubTrack.cpp
--- a/Engine/Source/Runtime/MovieScene/Private/Tracks/MovieSceneSubTrack.cpp
+++ b/Engine/Source/Runtime/MovieScene/Private/Tracks/MovieSceneSubTrack.cpp
@@ -64,9 +64,9 @@ UMovieSceneSubSection* UMovieSceneSubTrack::AddSequenceOnRow(UMovieSceneSequence
 	Sections.Add(NewSection);
 
 #if WITH_EDITORONLY_DATA
-	if (Sequence && Sequence->GetMovieScene())
+	if (Sequence)
 	{
-		NewSection->TimecodeSource = Sequence->GetMovieScene()->TimecodeSource;
+		NewSection->TimecodeSource = Sequence->GetEarliestTimecodeSource();
 	}
 #endif
 
