Takes: set TimecodeSource on UMovieSceneTakeSection when processing a recorded take

When recording a take with timecode in Take Recorder, the UMovieSceneTakeSection
that stores the timecode data is not created until
UTakeRecorderActorSource::ProcessRecordedTimes(), which happens when recording
stops. It therefore isn't available at the start of recording to have its timecode source
set with all of the other movie scene sections for the take. Previously, this was leaving
the TimecodeSource property of the UMovieSceneTakeSection in its default state,
inconsistent with the property values of all other sections for the take.

This adds handling at the end of ProcessRecordedTimes() to ensure that the TimecodeSource
property is set correctly on the take section, either re-using the collected timecode data
from recording, or by adopting the TimecodeSource of another one of the take's sections
if that timecode data is not available for some reason.

diff --git a/Engine/Plugins/VirtualProduction/Takes/Source/TakeRecorderSources/Private/TakeRecorderActorSource.cpp b/Engine/Plugins/VirtualProduction/Takes/Source/TakeRecorderSources/Private/TakeRecorderActorSource.cpp
--- a/Engine/Plugins/VirtualProduction/Takes/Source/TakeRecorderSources/Private/TakeRecorderActorSource.cpp
+++ b/Engine/Plugins/VirtualProduction/Takes/Source/TakeRecorderSources/Private/TakeRecorderActorSource.cpp
@@ -645,6 +645,11 @@ void UTakeRecorderActorSource::ProcessRecordedTimes(ULevelSequence* InSequence)
 		return;
 	}
 
+	// In case we need it later, get the earliest timecode source *before* we
+	// add the take section, since its timecode source will be default
+	// constructed as all zeros and might accidentally compare as earliest.
+	const FMovieSceneTimecodeSource EarliestTimecodeSource = MovieScene->GetEarliestTimecodeSource();
+
 	for (UMovieSceneTrack* Track : Binding->GetTracks())
 	{
 		for (UMovieSceneSection* Section : Track->GetAllSections())
@@ -730,6 +735,29 @@ void UTakeRecorderActorSource::ProcessRecordedTimes(ULevelSequence* InSequence)
 		TakeSection->SubFramesCurve.Set(Times, SubFrames);
 	}
 
+	// Since the take section was created post recording here in this
+	// function, it wasn't available at the start of recording to have
+	// its timecode source set with the other sections, so we set it here.
+	if (TakeSection->HoursCurve.GetNumKeys() > 0)
+	{
+		// We populated the take section's timecode curves with data, so
+		// use the first values as the timecode source.
+		const int32 Hours = TakeSection->HoursCurve.GetValues()[0];
+		const int32 Minutes = TakeSection->MinutesCurve.GetValues()[0];
+		const int32 Seconds = TakeSection->SecondsCurve.GetValues()[0];
+		const int32 Frames = TakeSection->FramesCurve.GetValues()[0];
+		const bool bIsDropFrame = false;
+		const FTimecode Timecode(Hours, Minutes, Seconds, Frames, bIsDropFrame);
+		TakeSection->TimecodeSource = FMovieSceneTimecodeSource(Timecode);
+	}
+	else
+	{
+		// Otherwise, adopt the earliest timecode source from one of the movie
+		// scene's other sections as the timecode source for the take section.
+		// This case is unlikely.
+		TakeSection->TimecodeSource = EarliestTimecodeSource;
+	}
+
 	if (UTakeMetaData* TakeMetaData = InSequence->FindMetaData<UTakeMetaData>())
 	{
 		TakeSection->Slate.SetDefault(FString::Printf(TEXT("%s_%d"), *TakeMetaData->GetSlate(), TakeMetaData->GetTakeNumber()));
