Take Recorder: Rework ETakeRecorderPanelMode so that it is derived by whatever data is set (ie. if 'SuppliedLevelSequence', then the mode is ReviewingRecording, if 'RecordingLevelSequence', then the mode is NewRecording, etc)

diff --git a/Engine/Plugins/VirtualProduction/Takes/Source/TakeRecorder/Private/Recorder/TakeRecorderPanel.cpp b/Engine/Plugins/VirtualProduction/Takes/Source/TakeRecorder/Private/Recorder/TakeRecorderPanel.cpp
--- a/Engine/Plugins/VirtualProduction/Takes/Source/TakeRecorder/Private/Recorder/TakeRecorderPanel.cpp
+++ b/Engine/Plugins/VirtualProduction/Takes/Source/TakeRecorder/Private/Recorder/TakeRecorderPanel.cpp
@@ -39,7 +39,8 @@ ETakeRecorderPanelMode UTakeRecorderPanel::GetMode() const
 		return ETakeRecorderPanelMode::NewRecording;
 	}
 
-	return WeakTabContent.Pin()->GetMode();
+	TOptional<ETakeRecorderPanelMode> Mode = WeakTabContent.Pin()->GetMode();
+	return Mode.IsSet() ? Mode.GetValue() : ETakeRecorderPanelMode::NewRecording;
 }
 
 void UTakeRecorderPanel::SetupForRecording_TakePreset(UTakePreset* TakePresetAsset)
diff --git a/Engine/Plugins/VirtualProduction/Takes/Source/TakeRecorder/Private/Widgets/STakeRecorderPanel.cpp b/Engine/Plugins/VirtualProduction/Takes/Source/TakeRecorder/Private/Widgets/STakeRecorderPanel.cpp
--- a/Engine/Plugins/VirtualProduction/Takes/Source/TakeRecorder/Private/Widgets/STakeRecorderPanel.cpp
+++ b/Engine/Plugins/VirtualProduction/Takes/Source/TakeRecorder/Private/Widgets/STakeRecorderPanel.cpp
@@ -6,6 +6,7 @@
 #include "Widgets/SLevelSequenceTakeEditor.h"
 #include "Recorder/TakeRecorder.h"
 #include "Recorder/TakeRecorderBlueprintLibrary.h"
+#include "Recorder/TakeRecorderPanel.h"
 #include "ScopedSequencerPanel.h"
 #include "ITakeRecorderModule.h"
 #include "TakesCoreBlueprintLibrary.h"
@@ -600,6 +601,24 @@ void STakeRecorderPanel::ClearPendingTake()
 	}
 }
 
+TOptional<ETakeRecorderPanelMode> STakeRecorderPanel::GetMode() const
+{
+	if (SuppliedLevelSequence)
+	{
+		return ETakeRecorderPanelMode::ReviewingRecording;	
+	}
+	else if (RecordingLevelSequence)
+	{
+		return ETakeRecorderPanelMode::NewRecording;	
+	}
+	else if (RecordIntoLevelSequence)
+	{
+		return ETakeRecorderPanelMode::RecordingInto;	
+	}
+
+	return TOptional<ETakeRecorderPanelMode>();
+}
+
 UTakePreset* STakeRecorderPanel::AllocateTransientPreset()
 {
 	return UTakePreset::AllocateTransientPreset(GetDefault<UTakeRecorderUserSettings>()->LastOpenedPreset.Get());
diff --git a/Engine/Plugins/VirtualProduction/Takes/Source/TakeRecorder/Private/Widgets/STakeRecorderPanel.h b/Engine/Plugins/VirtualProduction/Takes/Source/TakeRecorder/Private/Widgets/STakeRecorderPanel.h
--- a/Engine/Plugins/VirtualProduction/Takes/Source/TakeRecorder/Private/Widgets/STakeRecorderPanel.h
+++ b/Engine/Plugins/VirtualProduction/Takes/Source/TakeRecorder/Private/Widgets/STakeRecorderPanel.h
@@ -9,6 +9,7 @@
 
 enum class ECheckBoxState : uint8;
 enum class ETakeRecorderMode : uint8;
+enum class ETakeRecorderPanelMode : uint8;
 
 struct FAssetData;
 struct ITakeRecorderSourceTreeItem;
@@ -71,6 +72,8 @@ public:
 
 	void ClearPendingTake();
 
+	TOptional<ETakeRecorderPanelMode> GetMode() const;
+
 private:
 
 	virtual void AddReferencedObjects(FReferenceCollector& Collector) override;
diff --git a/Engine/Plugins/VirtualProduction/Takes/Source/TakeRecorder/Private/Widgets/STakeRecorderTabContent.cpp b/Engine/Plugins/VirtualProduction/Takes/Source/TakeRecorder/Private/Widgets/STakeRecorderTabContent.cpp
--- a/Engine/Plugins/VirtualProduction/Takes/Source/TakeRecorder/Private/Widgets/STakeRecorderTabContent.cpp
+++ b/Engine/Plugins/VirtualProduction/Takes/Source/TakeRecorder/Private/Widgets/STakeRecorderTabContent.cpp
@@ -43,7 +43,7 @@ void STakeRecorderTabContent::Construct(const FArguments& InArgs)
 
 EActiveTimerReturnType STakeRecorderTabContent::OnActiveTimer(double InCurrentTime, float InDeltaTime)
 {
-	if (!CurrentMode.IsSet())
+	if (!GetMode().IsSet())
 	{
 		SetupForRecording((UTakePreset*)nullptr);
 	}
@@ -60,16 +60,26 @@ const FSlateBrush* STakeRecorderTabContent::GetIcon() const
 	return IconAttribute.Get();
 }
 
-ETakeRecorderPanelMode STakeRecorderTabContent::GetMode() const
+TOptional<ETakeRecorderPanelMode> STakeRecorderTabContent::GetMode() const
 {
-	return CurrentMode.Get(ETakeRecorderPanelMode::NewRecording);
+	if (WeakPanel.IsValid())
+	{
+		return WeakPanel.Pin()->GetMode();
+	}
+	else if (WeakAssetEditor.IsValid())
+	{
+		return ETakeRecorderPanelMode::EditingPreset;
+	}
+	else
+	{
+		return TOptional<ETakeRecorderPanelMode>();
+	}
 }
 
 void STakeRecorderTabContent::SetupForRecording(UTakePreset* BasePreset)
 {
 	WeakAssetEditor = nullptr;
 
-	CurrentMode    = ETakeRecorderPanelMode::NewRecording;
 	TitleAttribute = ITakeRecorderModule::TakeRecorderTabLabel;
 	IconAttribute  = FAppStyle::Get().GetBrush("SequenceRecorder.TabIcon");
 
@@ -94,7 +104,6 @@ void STakeRecorderTabContent::SetupForRecording(ULevelSequence* LevelSequenceAss
 {
 	WeakAssetEditor = nullptr;
 
-	CurrentMode    = ETakeRecorderPanelMode::NewRecording;
 	TitleAttribute = ITakeRecorderModule::TakeRecorderTabLabel;
 	IconAttribute  = FTakeRecorderStyle::Get().GetBrush("TakeRecorder.TabIcon");
 
@@ -119,7 +128,6 @@ void STakeRecorderTabContent::SetupForRecordingInto(ULevelSequence* LevelSequenc
 {
 	WeakAssetEditor = nullptr;
 
-	CurrentMode    = ETakeRecorderPanelMode::NewRecording;
 	TitleAttribute = ITakeRecorderModule::TakeRecorderTabLabel;
 	IconAttribute  = FTakeRecorderStyle::Get().GetBrush("TakeRecorder.TabIcon");
 
@@ -144,7 +152,6 @@ void STakeRecorderTabContent::SetupForEditing(TSharedPtr<FTakePresetToolkit> InT
 {
 	WeakPanel = nullptr;
 
-	CurrentMode    = ETakeRecorderPanelMode::EditingPreset;
 	TitleAttribute = MakeAttributeSP(InToolkit.Get(), &FTakePresetToolkit::GetToolkitName);
 	IconAttribute  = MakeAttributeSP(InToolkit.Get(), &FTakePresetToolkit::GetTabIcon);
 
@@ -168,7 +175,6 @@ void STakeRecorderTabContent::SetupForViewing(ULevelSequence* LevelSequence)
 {
 	WeakAssetEditor = nullptr;
 
-	CurrentMode    = ETakeRecorderPanelMode::ReviewingRecording;
 	TitleAttribute = FText::FromString(LevelSequence->GetName());
 	IconAttribute  = FSlateIconFinder::FindIconBrushForClass(ULevelSequence::StaticClass());
 
diff --git a/Engine/Plugins/VirtualProduction/Takes/Source/TakeRecorder/Private/Widgets/STakeRecorderTabContent.h b/Engine/Plugins/VirtualProduction/Takes/Source/TakeRecorder/Private/Widgets/STakeRecorderTabContent.h
--- a/Engine/Plugins/VirtualProduction/Takes/Source/TakeRecorder/Private/Widgets/STakeRecorderTabContent.h
+++ b/Engine/Plugins/VirtualProduction/Takes/Source/TakeRecorder/Private/Widgets/STakeRecorderTabContent.h
@@ -27,7 +27,7 @@ public:
 
 	const FSlateBrush* GetIcon() const;
 
-	ETakeRecorderPanelMode GetMode() const;
+	TOptional<ETakeRecorderPanelMode> GetMode() const;
 
 	void SetupForRecording(ULevelSequence* LevelSequenceAsset);
 
@@ -69,7 +69,6 @@ private:
 
 private:
 
-	TOptional<ETakeRecorderPanelMode> CurrentMode;
 	TAttribute<FText> TitleAttribute;
 	TAttribute<const FSlateBrush*> IconAttribute;
 	TWeakPtr<STakeRecorderPanel> WeakPanel;
