MovieScene: remove TimecodeSource property stored directly on UMovieScene

These changes remove the last remnants of the TimecodeSource property that
was being stored directly on UMovieScene. UMovieScene is a container that may
represent multiple pieces of media as UMovieSceneSections, each of which may
have their own distinct TimecodeSource. The sections themselves manage setting
their own TimecodeSource when they are recorded or created from a piece of media.

Client code that is interested in the earliest TimecodeSource of any section in a movie
scene should use UMovieScene::GetEarliestTimecodeSource().

diff --git a/Engine/Plugins/VirtualProduction/Takes/Source/TakeRecorder/Private/SerializedRecorder.cpp b/Engine/Plugins/VirtualProduction/Takes/Source/TakeRecorder/Private/SerializedRecorder.cpp
--- a/Engine/Plugins/VirtualProduction/Takes/Source/TakeRecorder/Private/SerializedRecorder.cpp
+++ b/Engine/Plugins/VirtualProduction/Takes/Source/TakeRecorder/Private/SerializedRecorder.cpp
@@ -239,8 +239,6 @@ bool FSerializedRecorder::LoadSubSequenceFile(UMovieSceneSequence* InMovieSceneS
 			TargetSequence = UTakeRecorderSources::CreateSubSequenceForSource(InMasterSequence, SubSequenceName, SubSequenceName);
 			if (TargetSequence)
 			{
-				TargetSequence->GetMovieScene()->TimecodeSource = FApp::GetTimecode();
-
 				// If there's already a Subscene Track for our sub-sequence we need to remove that track before create a new one. No data is lost in this process as the
 				// sequence that the subscene points to has been copied by CreateSubSequenceForSource so a new track pointed to the new subsequence includes all the old data.
 				TOptional<int32> RowIndex;
diff --git a/Engine/Plugins/VirtualProduction/Takes/Source/TakesCore/Private/TakeRecorderSources.cpp b/Engine/Plugins/VirtualProduction/Takes/Source/TakesCore/Private/TakeRecorderSources.cpp
--- a/Engine/Plugins/VirtualProduction/Takes/Source/TakesCore/Private/TakeRecorderSources.cpp
+++ b/Engine/Plugins/VirtualProduction/Takes/Source/TakesCore/Private/TakeRecorderSources.cpp
@@ -333,7 +333,6 @@ void UTakeRecorderSources::StartRecording(class ULevelSequence* InSequence, cons
 
 	UE_LOG(LogTakesCore, Log, TEXT("StartRecording: %s"), *CurrentTimecode.ToString());
 
-	InSequence->GetMovieScene()->TimecodeSource = CurrentTimecode;
 	StartRecordingPreRecordedSources(InCurrentFrameTime);
 }
 
diff --git a/Engine/Source/Editor/SequenceRecorder/Private/SequenceRecorder.cpp b/Engine/Source/Editor/SequenceRecorder/Private/SequenceRecorder.cpp
--- a/Engine/Source/Editor/SequenceRecorder/Private/SequenceRecorder.cpp
+++ b/Engine/Source/Editor/SequenceRecorder/Private/SequenceRecorder.cpp
@@ -763,8 +763,6 @@ bool FSequenceRecorder::StartRecordingInternal(UWorld* World)
 			{
 				CurrentSequence = LevelSequence;
 
-				LevelSequence->GetMovieScene()->TimecodeSource = SequenceRecorderUtils::GetTimecodeSource();
-
 				FAssetRegistryModule::AssetCreated(LevelSequence);
 
 				RefreshNextSequence();
diff --git a/Engine/Source/Runtime/MovieScene/Public/MovieScene.h b/Engine/Source/Runtime/MovieScene/Public/MovieScene.h
--- a/Engine/Source/Runtime/MovieScene/Public/MovieScene.h
+++ b/Engine/Source/Runtime/MovieScene/Public/MovieScene.h
@@ -1015,10 +1015,6 @@ public:
 
 	UMovieSceneNodeGroupCollection& GetNodeGroups() { return *NodeGroupCollection; }
 
-	/** The timecode at which this movie scene section is based (ie. when it was recorded) */
-	UPROPERTY()
-	FMovieSceneTimecodeSource TimecodeSource;
-
 #endif	// WITH_EDITORONLY_DATA
 
 public:
