Take Recorder: Check for GEditor

diff --git a/Engine/Plugins/VirtualProduction/Takes/Source/TakeRecorderSources/Private/TakeRecorderWorldSource.cpp b/Engine/Plugins/VirtualProduction/Takes/Source/TakeRecorderSources/Private/TakeRecorderWorldSource.cpp
--- a/Engine/Plugins/VirtualProduction/Takes/Source/TakeRecorderSources/Private/TakeRecorderWorldSource.cpp
+++ b/Engine/Plugins/VirtualProduction/Takes/Source/TakeRecorderSources/Private/TakeRecorderWorldSource.cpp
@@ -128,7 +128,7 @@ bool UTakeRecorderWorldSource::CanAddSource(UTakeRecorderSources* InSources) con
 
 void UTakeRecorderWorldSource::AutotrackActors(class ULevelSequence* InSequence, UWorld* InWorld)
 {
-	if (!InWorld)
+	if (!InWorld || !GEditor)
 	{
 		return;
 	}
