Fix a crash when TakeRecoder would do concurrent calls to its SaveAsset function

diff --git a/Engine/Plugins/VirtualProduction/Takes/Source/TakesCore/Private/TakeRecorderSources.cpp b/Engine/Plugins/VirtualProduction/Takes/Source/TakesCore/Private/TakeRecorderSources.cpp
--- a/Engine/Plugins/VirtualProduction/Takes/Source/TakesCore/Private/TakeRecorderSources.cpp
+++ b/Engine/Plugins/VirtualProduction/Takes/Source/TakesCore/Private/TakeRecorderSources.cpp
@@ -500,14 +500,13 @@ void UTakeRecorderSources::StopRecording(class ULevelSequence* InSequence, const
 	TArray<UObject*> AssetsToCleanUp;
 	for (auto SourceSubSequence : SourceSubSequenceMap)
 	{
-		if (SourceSubSequence.Value)
+		// Only save subsequences but not the master sequence, it is already saved in UTakeRecorder::StopInternal
+		const bool bIsValidSubSequence = SourceSubSequence.Value && SourceSubSequence.Value.Get() != InSequence;
+		if (bIsValidSubSequence)
 		{
 			if (bCancelled)
 			{
-				if (SourceSubSequence.Value.Get() != InSequence)
-				{
-					AssetsToCleanUp.Add(SourceSubSequence.Value.Get());
-				}
+				AssetsToCleanUp.Add(SourceSubSequence.Value.Get());
 			}
 			else if (Settings.bSaveRecordedAssets)
 			{
