Take Recorder: Stop recording when the editor is closed

diff --git a/Engine/Plugins/VirtualProduction/Takes/Source/TakeRecorder/Private/TakeRecorderModule.cpp b/Engine/Plugins/VirtualProduction/Takes/Source/TakeRecorder/Private/TakeRecorderModule.cpp
--- a/Engine/Plugins/VirtualProduction/Takes/Source/TakeRecorder/Private/TakeRecorderModule.cpp
+++ b/Engine/Plugins/VirtualProduction/Takes/Source/TakeRecorder/Private/TakeRecorderModule.cpp
@@ -8,6 +8,7 @@
 #include "TakeRecorderCommands.h"
 #include "TakeRecorderStyle.h"
 #include "TakePresetActions.h"
+#include "Recorder/TakeRecorder.h"
 
 #include "WorkspaceMenuStructureModule.h"
 #include "WorkspaceMenuStructure.h"
@@ -279,12 +280,30 @@ void FTakeRecorderModule::StartupModule()
 			FCoreDelegates::OnPostEngineInit.AddRaw(this, &FTakeRecorderModule::RegisterMenus);
 		}
 	}
+
+	if (GEditor)
+	{
+		GEditor->OnEditorClose().AddRaw(this, &FTakeRecorderModule::OnEditorClose);
+	}
 #endif
 }
 
+void FTakeRecorderModule::OnEditorClose()
+{
+	if (UTakeRecorder* ActiveRecorder = UTakeRecorder::GetActiveRecorder())
+	{
+		ActiveRecorder->Stop();
+	}
+}
+
 void FTakeRecorderModule::ShutdownModule()
 {
 #if WITH_EDITOR
+	if (GEditor)
+	{
+		GEditor->OnEditorClose().RemoveAll(this);
+	}
+
 	FCoreDelegates::OnPostEngineInit.RemoveAll(this);
 #endif
 
diff --git a/Engine/Plugins/VirtualProduction/Takes/Source/TakeRecorder/Private/TakeRecorderModule.h b/Engine/Plugins/VirtualProduction/Takes/Source/TakeRecorder/Private/TakeRecorderModule.h
--- a/Engine/Plugins/VirtualProduction/Takes/Source/TakeRecorder/Private/TakeRecorderModule.h
+++ b/Engine/Plugins/VirtualProduction/Takes/Source/TakeRecorder/Private/TakeRecorderModule.h
@@ -78,6 +78,8 @@ private:
 
 	void RegisterMenus();
 
+	void OnEditorClose();
+
 private:
 
 	DECLARE_MULTICAST_DELEGATE_TwoParams(FOnExtendSourcesMenuEvent, TSharedRef<FExtender>, UTakeRecorderSources*);
