Take Recorder: Disregard cameras in added actor sources when automatically creating camera cut tracks since they were not explicitly added by the user.

This fixes a case where if you add a cine camera actor as a source and some other actor gets added automatically (because of hierarchy), that othr actor's camera component might become the camera cut.

diff --git a/Engine/Plugins/VirtualProduction/Takes/Source/TakeRecorderSources/Private/TakeRecorderActorSource.cpp b/Engine/Plugins/VirtualProduction/Takes/Source/TakeRecorderSources/Private/TakeRecorderActorSource.cpp
--- a/Engine/Plugins/VirtualProduction/Takes/Source/TakeRecorderSources/Private/TakeRecorderActorSource.cpp
+++ b/Engine/Plugins/VirtualProduction/Takes/Source/TakeRecorderSources/Private/TakeRecorderActorSource.cpp
@@ -813,7 +813,7 @@ TArray<UTakeRecorderSource*> UTakeRecorderActorSource::PostRecording(ULevelSeque
 		// Expand the Movie Scene Playback Range to encompass all of the sections now that they've all been created.
 		SequenceRecorderUtils::ExtendSequencePlaybackRange(InSequence);
 
-		if (Target.IsValid())
+		if (Target.IsValid() && !ParentSource)
 		{
 			// Automatically add or update the camera cut track if there is a camera component
 			AActor* TargetActor = Target.Get();
@@ -853,6 +853,8 @@ void UTakeRecorderActorSource::FinalizeRecording()
 	// Null these out there and NOT in PostRecording because they are used for cross sequence object binding via GetLevelSequenceID in PostRecording
 	TargetLevelSequence = nullptr;
 	MasterLevelSequence = nullptr;
+
+	ParentSource = nullptr;
 }
 
 void UTakeRecorderActorSource::PostProcessTrackRecorders(ULevelSequence* InSequence)
@@ -1593,6 +1595,7 @@ void UTakeRecorderActorSource::CreateNewActorSourceForReferencedActors()
 		// We don't use AddSource on the UTakeRecorderSources because this is called from functions that also adds returned items to the Source List. This prevents a 
 		// double add from occuring.
 		UTakeRecorderActorSource* ActorSource = NewObject<UTakeRecorderActorSource>(SourcesList, UTakeRecorderActorSource::StaticClass(), NAME_None, RF_Transactional);
+		ActorSource->ParentSource = this;
 
 		// We add it both to the local list (in case we need to start recording immediately) and to the class list
 		// so that we can clean up the recording when we finish.
diff --git a/Engine/Plugins/VirtualProduction/Takes/Source/TakeRecorderSources/Private/TakeRecorderActorSource.h b/Engine/Plugins/VirtualProduction/Takes/Source/TakeRecorderSources/Private/TakeRecorderActorSource.h
--- a/Engine/Plugins/VirtualProduction/Takes/Source/TakeRecorderSources/Private/TakeRecorderActorSource.h
+++ b/Engine/Plugins/VirtualProduction/Takes/Source/TakeRecorderSources/Private/TakeRecorderActorSource.h
@@ -92,6 +92,9 @@ public:
 	UPROPERTY()
 	TArray<TObjectPtr<class UMovieSceneTrackRecorder>> TrackRecorders;
 
+	/** The parent actor source that generated this actor source (ie. through parenting or as an attached component). Null after PostRecording */
+	UPROPERTY()
+	TObjectPtr<UTakeRecorderActorSource> ParentSource;
 
 public:
 
