Correct the mismatch between setting and reseting the pointer to UTakeRecorderPanel

diff --git a/Engine/Plugins/VirtualProduction/Takes/Source/TakeRecorder/Private/Widgets/STakeRecorderPanel.cpp b/Engine/Plugins/VirtualProduction/Takes/Source/TakeRecorder/Private/Widgets/STakeRecorderPanel.cpp
--- a/Engine/Plugins/VirtualProduction/Takes/Source/TakeRecorder/Private/Widgets/STakeRecorderPanel.cpp
+++ b/Engine/Plugins/VirtualProduction/Takes/Source/TakeRecorder/Private/Widgets/STakeRecorderPanel.cpp
@@ -5,7 +5,6 @@
 #include "Widgets/TakeRecorderWidgetConstants.h"
 #include "Widgets/SLevelSequenceTakeEditor.h"
 #include "Recorder/TakeRecorder.h"
-#include "Recorder/TakeRecorderBlueprintLibrary.h"
 #include "Recorder/TakeRecorderPanel.h"
 #include "ScopedSequencerPanel.h"
 #include "ITakeRecorderModule.h"
@@ -75,7 +74,6 @@ STakeRecorderPanel::~STakeRecorderPanel()
 	ITakeRecorderModule& TakeRecorderModule = FModuleManager::Get().LoadModuleChecked<ITakeRecorderModule>("TakeRecorder");
 	TakeRecorderModule.OnForceSaveAsPreset().Unbind();
 	TakeRecorderModule.GetExternalObjectAddRemoveEventDelegate().Remove(OnWidgetExternalObjectChangedHandle);
-	UTakeRecorderBlueprintLibrary::SetTakeRecorderPanel(nullptr);
 }
 
 PRAGMA_DISABLE_OPTIMIZATION
diff --git a/Engine/Plugins/VirtualProduction/Takes/Source/TakeRecorder/Private/Widgets/STakeRecorderTabContent.cpp b/Engine/Plugins/VirtualProduction/Takes/Source/TakeRecorder/Private/Widgets/STakeRecorderTabContent.cpp
--- a/Engine/Plugins/VirtualProduction/Takes/Source/TakeRecorder/Private/Widgets/STakeRecorderTabContent.cpp
+++ b/Engine/Plugins/VirtualProduction/Takes/Source/TakeRecorder/Private/Widgets/STakeRecorderTabContent.cpp
@@ -41,6 +41,11 @@ void STakeRecorderTabContent::Construct(const FArguments& InArgs)
 	}
 }
 
+STakeRecorderTabContent::~STakeRecorderTabContent()
+{
+	UTakeRecorderBlueprintLibrary::SetTakeRecorderPanel(nullptr);
+}
+
 EActiveTimerReturnType STakeRecorderTabContent::OnActiveTimer(double InCurrentTime, float InDeltaTime)
 {
 	if (!GetMode().IsSet())
diff --git a/Engine/Plugins/VirtualProduction/Takes/Source/TakeRecorder/Private/Widgets/STakeRecorderTabContent.h b/Engine/Plugins/VirtualProduction/Takes/Source/TakeRecorder/Private/Widgets/STakeRecorderTabContent.h
--- a/Engine/Plugins/VirtualProduction/Takes/Source/TakeRecorder/Private/Widgets/STakeRecorderTabContent.h
+++ b/Engine/Plugins/VirtualProduction/Takes/Source/TakeRecorder/Private/Widgets/STakeRecorderTabContent.h
@@ -21,6 +21,8 @@ public:
 	SLATE_BEGIN_ARGS(STakeRecorderTabContent){}
 	SLATE_END_ARGS()
 
+	~STakeRecorderTabContent();
+
 	void Construct(const FArguments& InArgs);
 
 	FText GetTitle() const;
