Take Recorder: Don't replace if recording as spawnable so that recording as spawnable will always create a new spawnable

diff --git a/Engine/Plugins/VirtualProduction/Takes/Source/TakeRecorderSources/Private/TakeRecorderActorSource.cpp b/Engine/Plugins/VirtualProduction/Takes/Source/TakeRecorderSources/Private/TakeRecorderActorSource.cpp
--- a/Engine/Plugins/VirtualProduction/Takes/Source/TakeRecorderSources/Private/TakeRecorderActorSource.cpp
+++ b/Engine/Plugins/VirtualProduction/Takes/Source/TakeRecorderSources/Private/TakeRecorderActorSource.cpp
@@ -1324,16 +1324,8 @@ FGuid UTakeRecorderActorSource::ResolveActorFromSequence(AActor* InActor, ULevel
 {
 	UMovieScene* MovieScene = CurrentSequence->GetMovieScene();
 
-	// Look through all Spawnables and Possessables in the sequence to see if there's one with the same name as the actor
-	for (int32 SpawnableCount = 0; SpawnableCount < MovieScene->GetSpawnableCount(); ++SpawnableCount)
-	{
-		const FMovieSceneSpawnable& Spawnable = MovieScene->GetSpawnable(SpawnableCount);
-		if (Spawnable.GetName() == InActor->GetActorLabel())
-		{
-			return Spawnable.GetGuid();
-		}
-	}
-
+	// Look through all Possessables in the sequence to see if there's one with the same name as the actor. We purposely do not look at 
+	// Spawnables so that recording as spawnable will always create a new spawnable.
 	for (int32 PossessableCount = 0; PossessableCount < MovieScene->GetPossessableCount(); ++PossessableCount)
 	{
 		const FMovieScenePossessable& Possessable = MovieScene->GetPossessable(PossessableCount);
@@ -1343,7 +1335,7 @@ FGuid UTakeRecorderActorSource::ResolveActorFromSequence(AActor* InActor, ULevel
 		}
 	}
 
-	// There's no Spawnable or Possessable with the same name as the actor, so this actor hasn't been added
+	// There's no Possessable with the same name as the actor, so this actor hasn't been added
 	// to the sequence yet.
 	return FGuid();
 }
@@ -1400,8 +1392,8 @@ void UTakeRecorderActorSource::CleanExistingDataFromSequence(const FGuid& ForGui
 {
 	if (ForGuid.IsValid())
 	{
-		// Check to see if there is a Possessable or Spawnable in this sequence with the specified Guid and remove their old data if needed.
-		// Removing the Spawnable/Possessable will remove their bindings which will remove the associated tracks and their data as well.
+		// Check to see if there is a Possessable in this sequence with the specified Guid and remove their old data if needed.
+		// Removing the Possessable will remove their bindings which will remove the associated tracks and their data as well.
 		UMovieScene* MovieScene = InSequence.GetMovieScene();
 
 		TArray<FGuid> OutChildGuids;
diff --git a/Engine/Plugins/VirtualProduction/Takes/Source/TakeRecorderSources/Private/TakeRecorderActorSource.h b/Engine/Plugins/VirtualProduction/Takes/Source/TakeRecorderSources/Private/TakeRecorderActorSource.h
--- a/Engine/Plugins/VirtualProduction/Takes/Source/TakeRecorderSources/Private/TakeRecorderActorSource.h
+++ b/Engine/Plugins/VirtualProduction/Takes/Source/TakeRecorderSources/Private/TakeRecorderActorSource.h
@@ -192,10 +192,10 @@ protected:
 	/** Update our cached properties for what will be recorded. Done here so the UI doesn't have to iterate through map every frame. */
 	void UpdateCachedNumberOfRecordedProperties();
 
-	/** Returns the Guid of the Spawnable/Possessable in the specified sequence that represents the given actor, or an invalid Guid if the actor has no object binding in the sequence. */
+	/** Returns the Guid of the Possessable in the specified sequence that represents the given actor, or an invalid Guid if the actor has no object binding in the sequence. */
 	FGuid ResolveActorFromSequence(AActor* InActor, ULevelSequence* CurrentSequence) const;
 
-	/** Remove the Spawnable/Possessable data for the given Guid from the sequence. Calls CleanExistingDataFromSequenceImpl afterwards for any other cleanup you may wish to do. */
+	/** Remove the Possessable data for the given Guid from the sequence. Calls CleanExistingDataFromSequenceImpl afterwards for any other cleanup you may wish to do. */
 	void CleanExistingDataFromSequence(const FGuid& ForGuid, ULevelSequence& InSequence);
 
 	/** Called as part of PostRecording before Track Recorders are finalized. Calls PostProcessTrackRecordersImpl afterwards for any other post processing you wish to do before Track recorders are finalized. */
