Take Recorder: Check for valid component before recording. This can happen if construction scripts are run.

diff --git a/Engine/Plugins/VirtualProduction/Takes/Source/TakeRecorderSources/Private/TakeRecorderActorSource.cpp b/Engine/Plugins/VirtualProduction/Takes/Source/TakeRecorderSources/Private/TakeRecorderActorSource.cpp
--- a/Engine/Plugins/VirtualProduction/Takes/Source/TakeRecorderSources/Private/TakeRecorderActorSource.cpp
+++ b/Engine/Plugins/VirtualProduction/Takes/Source/TakeRecorderSources/Private/TakeRecorderActorSource.cpp
@@ -563,9 +563,9 @@ void UTakeRecorderActorSource::TickRecording(const FQualifiedFrameTime& CurrentS
 	TArray<UObject*> TraversedObjects;
 	for (UActorComponent* AddedComponent : NewComponentsAdded)
 	{
-		if (Target.IsValid())
+		if (Target.IsValid() && ::IsValid(AddedComponent))
 		{
-			UE_LOG(LogTakesCore, Log, TEXT("Detected newly added component %s on Actor %s, begining to record component's properties now."), *AddedComponent->GetReadableName(), *Target->GetName());
+			UE_LOG(LogTakesCore, Log, TEXT("Detected newly added component %s on Actor %s, beginning to record component's properties now."), *AddedComponent->GetReadableName(), *Target->GetName());
 			TSet<UMovieSceneTrackRecorder*> PreviousTrackRecorders = TSet<UMovieSceneTrackRecorder*>(TrackRecorders);
 
 			// We should create a new property map attached to the right parent, and then initialize it using existing flow. This works for Possessables too as it will throw a warning that
