Take Recorder: Fix static analysis

diff --git a/Engine/Plugins/VirtualProduction/Takes/Source/TakeTrackRecorders/Private/MovieSceneAnimationTrackRecorder.cpp b/Engine/Plugins/VirtualProduction/Takes/Source/TakeTrackRecorders/Private/MovieSceneAnimationTrackRecorder.cpp
--- a/Engine/Plugins/VirtualProduction/Takes/Source/TakeTrackRecorders/Private/MovieSceneAnimationTrackRecorder.cpp
+++ b/Engine/Plugins/VirtualProduction/Takes/Source/TakeTrackRecorders/Private/MovieSceneAnimationTrackRecorder.cpp
@@ -272,7 +272,11 @@ void UMovieSceneAnimationTrackRecorder::RecordSampleImpl(const FQualifiedFrameTi
 		//If not removing root we also don't record in world space ( not totally sure if it matters but matching up with Sequence Recorder)
 		bool bRecordInWorldSpace = bRootWasRemoved == false ? false : true;
 
-		FTrackRecorderSettings TrackRecorderSettings = OwningTakeRecorderSource->GetTrackRecorderSettings();
+		FTrackRecorderSettings TrackRecorderSettings;
+		if (OwningTakeRecorderSource)
+		{
+			TrackRecorderSettings = OwningTakeRecorderSource->GetTrackRecorderSettings();
+		}
 
 		if (bRecordInWorldSpace && AttachParent && OwningTakeRecorderSource)
 		{
