Take Recorder: Add a user option to start at the current playhead as opposed to restarting at the playback start frame. The current default behavior is unchanged.

This allows a user to scrub to a frame in the pending take, hit record, and recording will start from there. The resulting recording's playback range will start from that same frame. When the user stops recording, the original playback range will be restored.

diff --git a/Engine/Plugins/VirtualProduction/Takes/Source/TakeRecorder/Private/Recorder/TakeRecorder.cpp b/Engine/Plugins/VirtualProduction/Takes/Source/TakeRecorder/Private/Recorder/TakeRecorder.cpp
--- a/Engine/Plugins/VirtualProduction/Takes/Source/TakeRecorder/Private/Recorder/TakeRecorder.cpp
+++ b/Engine/Plugins/VirtualProduction/Takes/Source/TakeRecorder/Private/Recorder/TakeRecorder.cpp
@@ -513,6 +513,13 @@ bool UTakeRecorder::Initialize ( ULevelSequence* LevelSequenceBase, UTakeRecorde
 	TSharedPtr<ISequencer> Sequencer = WeakSequencer.Pin();
 	if (Sequencer.IsValid())
 	{
+		// If a start frame was specified, adjust the playback range before rewinding to the beginning of the playback range
+		UMovieScene* MovieScene = SequenceAsset->GetMovieScene();
+		MovieScene->SetPlaybackRange(TRange<FFrameNumber>(Parameters.StartFrame, MovieScene->GetPlaybackRange().GetUpperBoundValue()));
+
+		// Always start the recording at the beginning of the playback range
+		Sequencer->SetLocalTime(MovieScene->GetPlaybackRange().GetLowerBoundValue());
+
 		if (Parameters.TakeRecorderMode == ETakeRecorderMode::RecordNewSequence)
 		{
 			USequencerSettings* SequencerSettings = USequencerSettingsContainer::GetOrCreate<USequencerSettings>(TEXT("TakeRecorderSequenceEditor"));
@@ -530,7 +537,6 @@ bool UTakeRecorder::Initialize ( ULevelSequence* LevelSequenceBase, UTakeRecorde
 		// Center the view range around the current time about to be captured
 		FAnimatedRange Range = Sequencer->GetViewRange();
 		FTimecode CurrentTime = FApp::GetTimecode();
-		UMovieScene* MovieScene = SequenceAsset->GetMovieScene();
 		FFrameRate FrameRate = MovieScene->GetDisplayRate();
 		FFrameNumber ViewRangeStart = CurrentTime.ToFrameNumber(FrameRate);
 		double ViewRangeStartSeconds = Parameters.Project.bStartAtCurrentTimecode ? FrameRate.AsSeconds(ViewRangeStart) : MovieScene->GetPlaybackRange().GetLowerBoundValue() / MovieScene->GetTickResolution();
diff --git a/Engine/Plugins/VirtualProduction/Takes/Source/TakeRecorder/Private/Recorder/TakeRecorderParameters.cpp b/Engine/Plugins/VirtualProduction/Takes/Source/TakeRecorder/Private/Recorder/TakeRecorderParameters.cpp
--- a/Engine/Plugins/VirtualProduction/Takes/Source/TakeRecorder/Private/Recorder/TakeRecorderParameters.cpp
+++ b/Engine/Plugins/VirtualProduction/Takes/Source/TakeRecorder/Private/Recorder/TakeRecorderParameters.cpp
@@ -7,6 +7,7 @@ FTakeRecorderUserParameters::FTakeRecorderUserParameters()
 	: bMaximizeViewport(false)
 	, CountdownSeconds(0.f)
 	, EngineTimeDilation(1.f)
+	, bResetPlayhead(true)
 	, bStopAtPlaybackEnd(false)
 	, bRemoveRedundantTracks(true)
 	, ReduceKeysTolerance(KINDA_SMALL_NUMBER)
diff --git a/Engine/Plugins/VirtualProduction/Takes/Source/TakeRecorder/Private/Widgets/STakeRecorderCockpit.cpp b/Engine/Plugins/VirtualProduction/Takes/Source/TakeRecorder/Private/Widgets/STakeRecorderCockpit.cpp
--- a/Engine/Plugins/VirtualProduction/Takes/Source/TakeRecorder/Private/Widgets/STakeRecorderCockpit.cpp
+++ b/Engine/Plugins/VirtualProduction/Takes/Source/TakeRecorder/Private/Widgets/STakeRecorderCockpit.cpp
@@ -57,6 +57,10 @@
 #include "LevelEditor.h"
 #include "ISettingsModule.h"
 #include "Dialog/SMessageDialog.h"
+#include "Subsystems/AssetEditorSubsystem.h"
+
+#include "ISequencer.h"
+#include "ILevelSequenceEditorToolkit.h"
 
 extern UNREALED_API UEditorEngine* GEditor;
 
@@ -1045,9 +1049,23 @@ void STakeRecorderCockpit::StartRecording()
 		Parameters.User    = GetDefault<UTakeRecorderUserSettings>()->Settings;
 		Parameters.Project = GetDefault<UTakeRecorderProjectSettings>()->Settings;
 		Parameters.TakeRecorderMode = TakeRecorderModeAttribute.Get();
+		Parameters.StartFrame = LevelSequence->GetMovieScene()->GetPlaybackRange().GetLowerBoundValue();
 
 		FText ErrorText = LOCTEXT("UnknownError", "An unknown error occurred when trying to start recording");
 
+		IAssetEditorInstance* AssetEditor = GEditor->GetEditorSubsystem<UAssetEditorSubsystem>()->FindEditorForAsset(LevelSequence, false);
+		ILevelSequenceEditorToolkit* LevelSequenceEditor = static_cast<ILevelSequenceEditorToolkit*>(AssetEditor);
+
+		if (LevelSequenceEditor && LevelSequenceEditor->GetSequencer())
+		{
+			// If not resetting the playhead, store the current time as the start frame for recording. 
+			// This will ultimately be the start of the playback range and the recording will begin from that time.
+			if (!Parameters.User.bResetPlayhead)
+			{
+				Parameters.StartFrame = LevelSequenceEditor->GetSequencer()->GetLocalTime().Time.FrameNumber;
+			}
+		}
+
 		UTakeRecorder* NewRecorder = NewObject<UTakeRecorder>(GetTransientPackage(), NAME_None, RF_Transient);
 
 		if (!NewRecorder->Initialize(LevelSequence, Sources, TakeMetaData, Parameters, &ErrorText))
diff --git a/Engine/Plugins/VirtualProduction/Takes/Source/TakeRecorder/Public/Recorder/TakeRecorderParameters.h b/Engine/Plugins/VirtualProduction/Takes/Source/TakeRecorder/Public/Recorder/TakeRecorderParameters.h
--- a/Engine/Plugins/VirtualProduction/Takes/Source/TakeRecorder/Public/Recorder/TakeRecorderParameters.h
+++ b/Engine/Plugins/VirtualProduction/Takes/Source/TakeRecorder/Public/Recorder/TakeRecorderParameters.h
@@ -27,8 +27,12 @@ struct FTakeRecorderUserParameters
 	UPROPERTY(config, EditAnywhere, BlueprintReadWrite, Category="User Settings", meta=(Units=Multiplier, ClampMin="0.00001", UIMin="0.00001"))
 	float EngineTimeDilation;
 
+	/** Reset playhead to beginning of the playback range when starting recording */
+	UPROPERTY(config, EditAnywhere, BlueprintReadWrite, Category = "User Settings")
+	bool bResetPlayhead;
+
 	/** Automatically stop recording when reaching the end of the playback range */
-	UPROPERTY(config, EditAnywhere, BlueprintReadWrite, Category="User Settings")
+	UPROPERTY(config, EditAnywhere, BlueprintReadWrite, Category = "User Settings")
 	bool bStopAtPlaybackEnd;
 
 	/** Recommended for use with recorded spawnables. Beware that changes to actor instances in the map after recording may alter the recording when played back */
@@ -157,6 +161,9 @@ struct FTakeRecorderParameters
 	UPROPERTY(BlueprintReadWrite, Category="Take Recorder")
 	ETakeRecorderMode TakeRecorderMode;
 
+	UPROPERTY(BlueprintReadWrite, Category = "Take Recorder")
+	FFrameNumber StartFrame;
+
 	/**
 	 * Option to disable recording and saving of data. This can be used in a scenario where multiple clients are running
 	 * take recorder, but only certain ones are set to process and save the data.
