Take Recorder: Fix time scaled recording so that the pending take is also played back at the dilated time.

Set the MatineeTimeDilation instead of just the TimeDilation since GetEffectiveTimeDilation() is the composition of both.
Move applying the engine time dilation for recording from InitializeFromParameters() to PreRecord() so that the recording countdown isn't slowed down.

diff --git a/Engine/Plugins/VirtualProduction/Takes/Source/TakeRecorder/Private/Recorder/TakeRecorder.cpp b/Engine/Plugins/VirtualProduction/Takes/Source/TakeRecorder/Private/Recorder/TakeRecorder.cpp
--- a/Engine/Plugins/VirtualProduction/Takes/Source/TakeRecorder/Private/Recorder/TakeRecorder.cpp
+++ b/Engine/Plugins/VirtualProduction/Takes/Source/TakeRecorder/Private/Recorder/TakeRecorder.cpp
@@ -724,28 +724,6 @@ void UTakeRecorder::InitializeFromParameters()
 			OnStopCleanup.Add(RestoreImmersiveMode);
 		}
 	}
-
-	// Apply engine Time Dilation
-	UWorld* RecordingWorld = GetWorld();
-	check(RecordingWorld);
-	if (AWorldSettings* WorldSettings = RecordingWorld->GetWorldSettings())
-	{
-		const float ExistingTimeDilation = WorldSettings->TimeDilation;
-		if (Parameters.User.EngineTimeDilation != ExistingTimeDilation)
-		{
-			WorldSettings->SetTimeDilation(Parameters.User.EngineTimeDilation);
-
-			// Restore it when we're done
-			auto RestoreTimeDilation = [ExistingTimeDilation, WeakWorldSettings = MakeWeakObjectPtr(WorldSettings)]
-			{
-				if (AWorldSettings* CleaupWorldSettings = WeakWorldSettings.Get())
-				{
-					CleaupWorldSettings->SetTimeDilation(ExistingTimeDilation);
-				}
-			};
-			OnStopCleanup.Add(RestoreTimeDilation);
-		}
-	}
 }
 
 bool UTakeRecorder::ShouldShowNotifications()
@@ -938,6 +916,29 @@ void UTakeRecorder::PreRecord()
 	{
 		Sequencer->RefreshTree();
 	}
+
+	// Apply engine Time Dilation after the countdown, otherwise the countdown will be dilated as well!
+	UWorld* RecordingWorld = GetWorld();
+	check(RecordingWorld);
+	if (AWorldSettings* WorldSettings = RecordingWorld->GetWorldSettings())
+	{
+		const float ExistingMatineeTimeDilation = WorldSettings->MatineeTimeDilation;
+
+		if (Parameters.User.EngineTimeDilation != ExistingMatineeTimeDilation)
+		{
+			WorldSettings->MatineeTimeDilation = Parameters.User.EngineTimeDilation;
+
+			// Restore it when we're done
+			auto RestoreTimeDilation = [ExistingMatineeTimeDilation, WeakWorldSettings = MakeWeakObjectPtr(WorldSettings)]
+			{
+				if (AWorldSettings* CleaupWorldSettings = WeakWorldSettings.Get())
+				{
+					CleaupWorldSettings->MatineeTimeDilation = ExistingMatineeTimeDilation;
+				}
+			};
+			OnStopCleanup.Add(RestoreTimeDilation);
+		}
+	}
 }
 
 void UTakeRecorder::Start()
