Expose functionality needed for caching, plus a few extra guards

diff --git a/Engine/Plugins/VirtualProduction/Takes/Source/TakeRecorder/Private/Recorder/TakeRecorder.cpp b/Engine/Plugins/VirtualProduction/Takes/Source/TakeRecorder/Private/Recorder/TakeRecorder.cpp
--- a/Engine/Plugins/VirtualProduction/Takes/Source/TakeRecorder/Private/Recorder/TakeRecorder.cpp
+++ b/Engine/Plugins/VirtualProduction/Takes/Source/TakeRecorder/Private/Recorder/TakeRecorder.cpp
@@ -568,13 +568,19 @@ void UTakeRecorder::DiscoverSourceWorld()
 	check(WorldToRecordIn);
 	WeakWorld = WorldToRecordIn;
 
-	UClass* Class = StaticLoadClass(UTakeRecorderOverlayWidget::StaticClass(), nullptr, TEXT("/Takes/UMG/DefaultRecordingOverlay.DefaultRecordingOverlay_C"));
-	if (Class)
+	// If CountdownSeconds is zero and the framerate is high, then we can create the overlay but it
+	// never ends up being visible. However, when the framerate is low (e.g. in debug) it does show
+	// for a single frame, which is undesirable, so only make it if it's going to last some time.
+	if (CountdownSeconds > 0)
 	{
-		OverlayWidget = CreateWidget<UTakeRecorderOverlayWidget>(WorldToRecordIn, Class);
-		OverlayWidget->SetFlags(RF_Transient);
-		OverlayWidget->SetRecorder(this);
-		OverlayWidget->AddToViewport();
+		UClass* Class = StaticLoadClass(UTakeRecorderOverlayWidget::StaticClass(), nullptr, TEXT("/Takes/UMG/DefaultRecordingOverlay.DefaultRecordingOverlay_C"));
+		if (Class)
+		{
+			OverlayWidget = CreateWidget<UTakeRecorderOverlayWidget>(WorldToRecordIn, Class);
+			OverlayWidget->SetFlags(RF_Transient);
+			OverlayWidget->SetRecorder(this);
+			OverlayWidget->AddToViewport();
+		}
 	}
 
 	bool bPlayInGame = WorldToRecordIn->WorldType == EWorldType::PIE || WorldToRecordIn->WorldType == EWorldType::Game;
diff --git a/Engine/Plugins/VirtualProduction/Takes/Source/TakeRecorderSources/Private/TakeRecorderActorSource.cpp b/Engine/Plugins/VirtualProduction/Takes/Source/TakeRecorderSources/Private/TakeRecorderActorSource.cpp
--- a/Engine/Plugins/VirtualProduction/Takes/Source/TakeRecorderSources/Private/TakeRecorderActorSource.cpp
+++ b/Engine/Plugins/VirtualProduction/Takes/Source/TakeRecorderSources/Private/TakeRecorderActorSource.cpp
@@ -184,6 +184,7 @@ UTakeRecorderActorSource::UTakeRecorderActorSource(const FObjectInitializer& Obj
 	RecordType = ETakeRecorderActorRecordType::ProjectDefault;
 	bReduceKeys = true;
 	bRecordParentHierarchy = true;
+	bShowProgressDialog = true;
 
 	TargetSequenceID = MovieSceneSequenceID::Invalid;
 
@@ -776,7 +777,7 @@ TArray<UTakeRecorderSource*> UTakeRecorderActorSource::PostRecording(ULevelSeque
 	Parameters.Project = GetDefault<UTakeRecorderProjectSettings>()->Settings;
 
 	FScopedSlowTask SlowTask((float)TrackRecorders.Num() + 1.0f, FText::Format(LOCTEXT("ProcessingActor", "Generating MovieScene Data for Actor {0}"), Target.IsValid() ? FText::FromString(Target.Get()->GetActorLabel()) : FText::GetEmpty()));
-	SlowTask.MakeDialog(false, true);
+	SlowTask.MakeDialog(false, bShowProgressDialog);
 
 	// We need to do some post-processing tasks on the Track Recorders (such as animation motion source fixup) so we do this now before finalizing
 	{
diff --git a/Engine/Plugins/VirtualProduction/Takes/Source/TakeRecorderSources/Private/TakeRecorderActorSource.h b/Engine/Plugins/VirtualProduction/Takes/Source/TakeRecorderSources/Private/TakeRecorderActorSource.h
--- a/Engine/Plugins/VirtualProduction/Takes/Source/TakeRecorderSources/Private/TakeRecorderActorSource.h
+++ b/Engine/Plugins/VirtualProduction/Takes/Source/TakeRecorderSources/Private/TakeRecorderActorSource.h
@@ -104,6 +104,10 @@ public:
 	UPROPERTY()
 	TObjectPtr<UTakeRecorderActorSource> ParentSource;
 
+	/** Show Dialog during the (possibly) slow parts of the take recording */
+	UPROPERTY()
+	bool bShowProgressDialog;
+
 public:
 
 	/*
diff --git a/Engine/Plugins/VirtualProduction/Takes/Source/TakeRecorderSources/Private/TakeRecorderSourceHelpers.cpp b/Engine/Plugins/VirtualProduction/Takes/Source/TakeRecorderSources/Private/TakeRecorderSourceHelpers.cpp
new file mode 100644
--- /dev/null
+++ b/Engine/Plugins/VirtualProduction/Takes/Source/TakeRecorderSources/Private/TakeRecorderSourceHelpers.cpp
@@ -0,0 +1,104 @@
+// Copyright Epic Games, Inc. All Rights Reserved.
+
+#include "TakeRecorderSourceHelpers.h"
+#include "TakeRecorderSource.h"
+#include "TakeRecorderActorSource.h"
+#include "LevelSequenceActor.h"
+#include "GameFramework/Actor.h"
+#include "ScopedTransaction.h"
+#include "TakeRecorderLevelSequenceSource.h"
+#include "Editor.h"
+
+#define LOCTEXT_NAMESPACE "TakeRecorderSourceHelpers"
+
+namespace TakeRecorderSourceHelpers
+{
+
+void AddActorSources(
+	UTakeRecorderSources* TakeRecorderSources,
+	TArrayView<AActor* const> ActorsToRecord,
+	bool                      bReduceKeys,
+	bool                      bShowProgress)
+{
+	if (ActorsToRecord.Num() > 0)
+	{
+		FScopedTransaction Transaction(FText::Format(
+			LOCTEXT("AddSources", "Add Recording {0}|plural(one=Source, other=Sources)"), ActorsToRecord.Num()));
+		TakeRecorderSources->Modify();
+
+		for (AActor* Actor : ActorsToRecord)
+		{
+			if (Actor->IsA<ALevelSequenceActor>())
+			{
+				ALevelSequenceActor* LevelSequenceActor = Cast<ALevelSequenceActor>(Actor);
+
+				UTakeRecorderLevelSequenceSource* LevelSequenceSource = nullptr;
+
+				for (UTakeRecorderSource* Source : TakeRecorderSources->GetSources())
+				{
+					if (Source->IsA<UTakeRecorderLevelSequenceSource>())
+					{
+						LevelSequenceSource = Cast<UTakeRecorderLevelSequenceSource>(Source);
+						break;
+					}
+				}
+
+				if (!LevelSequenceSource)
+				{
+					LevelSequenceSource = TakeRecorderSources->AddSource<UTakeRecorderLevelSequenceSource>();
+				}
+
+				ULevelSequence* Sequence = LevelSequenceActor->GetSequence();
+				if (Sequence)
+				{
+					if (!LevelSequenceSource->LevelSequencesToTrigger.Contains(Sequence))
+					{
+						LevelSequenceSource->LevelSequencesToTrigger.Add(Sequence);
+					}
+				}
+			}
+			else
+			{
+				UTakeRecorderActorSource* NewSource = TakeRecorderSources->AddSource<UTakeRecorderActorSource>();
+
+				if (AActor* EditorActor = EditorUtilities::GetEditorWorldCounterpartActor(Actor))
+				{
+					NewSource->Target = EditorActor;
+				}
+				else
+				{
+					NewSource->Target = Actor;
+				}
+				NewSource->bShowProgressDialog = bShowProgress;
+				NewSource->bReduceKeys = bReduceKeys;
+
+				// Send a PropertyChangedEvent so the class catches the callback and rebuilds the property map.
+				FPropertyChangedEvent PropertyChangedEvent(UTakeRecorderActorSource::StaticClass()->FindPropertyByName(GET_MEMBER_NAME_CHECKED(UTakeRecorderActorSource, Target)), EPropertyChangeType::ValueSet);
+				NewSource->PostEditChangeProperty(PropertyChangedEvent);
+			}
+		}
+	}
+}
+
+void RemoveAllActorSources(UTakeRecorderSources* TakeRecorderSources)
+{
+	if (TakeRecorderSources->GetSources().IsEmpty())
+	{
+		return;
+	}
+
+	FScopedTransaction Transaction(FText::Format(LOCTEXT(
+			"RemoveAllActorSources", 
+			"Remove Recording {0}|plural(one=Source, other=Sources)"), TakeRecorderSources->GetSources().Num()));
+	TakeRecorderSources->Modify();
+
+	while (!TakeRecorderSources->GetSources().IsEmpty())
+	{
+		TakeRecorderSources->RemoveSource(TakeRecorderSources->GetSources()[0]);
+	}
+}
+
+
+}
+
+#undef LOCTEXT_NAMESPACE // "TakeRecorderSourceHelpers"
\ No newline at end of file
diff --git a/Engine/Plugins/VirtualProduction/Takes/Source/TakeRecorderSources/Private/TakeRecorderSourcesModule.cpp b/Engine/Plugins/VirtualProduction/Takes/Source/TakeRecorderSources/Private/TakeRecorderSourcesModule.cpp
--- a/Engine/Plugins/VirtualProduction/Takes/Source/TakeRecorderSources/Private/TakeRecorderSourcesModule.cpp
+++ b/Engine/Plugins/VirtualProduction/Takes/Source/TakeRecorderSources/Private/TakeRecorderSourcesModule.cpp
@@ -12,6 +12,7 @@
 #include "TakesCoreBlueprintLibrary.h"
 #include "TakesCoreLog.h"
 #include "TakeMetaData.h"
+#include "TakeRecorderSourceHelpers.h"
 #include "TakeRecorderActorSource.h"
 #include "TakeRecorderSources.h"
 #include "TakeRecorderSourcesCommands.h"
@@ -63,66 +64,6 @@ namespace TakeRecorderSources
 #endif // WITH_EDITOR
 }
 
-static void AddActorSources(UTakeRecorderSources* Sources, TArrayView<AActor* const> InActors)
-{
-	if (InActors.Num() > 0)
-	{
-		FScopedTransaction Transaction(FText::Format(LOCTEXT("AddSources", "Add Recording {0}|plural(one=Source, other=Sources)"), InActors.Num()));
-		Sources->Modify();
-
-		for (AActor* Actor : InActors)
-		{
-			if (Actor->IsA<ALevelSequenceActor>())
-			{
-				ALevelSequenceActor* LevelSequenceActor = Cast<ALevelSequenceActor>(Actor);
-
-				UTakeRecorderLevelSequenceSource* LevelSequenceSource = nullptr;
-
-				for (UTakeRecorderSource* Source : Sources->GetSources())
-				{
-					if (Source->IsA<UTakeRecorderLevelSequenceSource>())
-					{
-						LevelSequenceSource = Cast<UTakeRecorderLevelSequenceSource>(Source);
-						break;
-					}
-				}
-
-				if (!LevelSequenceSource)
-				{
-					LevelSequenceSource = Sources->AddSource<UTakeRecorderLevelSequenceSource>();
-				}
-
-				ULevelSequence* Sequence = LevelSequenceActor->GetSequence();
-				if (Sequence)
-				{
-					if (!LevelSequenceSource->LevelSequencesToTrigger.Contains(Sequence))
-					{
-						LevelSequenceSource->LevelSequencesToTrigger.Add(Sequence);
-					}
-				}
-			}
-			else
-			{
-				UTakeRecorderActorSource* NewSource = Sources->AddSource<UTakeRecorderActorSource>();
-
-				if (AActor* EditorActor = EditorUtilities::GetEditorWorldCounterpartActor(Actor))
-				{
-					NewSource->Target = EditorActor;
-				}
-				else
-				{
-					NewSource->Target = Actor;
-				}
-
-				// Send a PropertyChangedEvent so the class catches the callback and rebuilds the property map.
-				FPropertyChangedEvent PropertyChangedEvent(UTakeRecorderActorSource::StaticClass()->FindPropertyByName(GET_MEMBER_NAME_CHECKED(UTakeRecorderActorSource, Target)), EPropertyChangeType::ValueSet);
-				NewSource->PostEditChangeProperty(PropertyChangedEvent);
-			}
-		}
-	}
-}
-
-
 namespace
 {
 	static AActor* FindActorByLabel(const FString& ActorNameStr, UWorld* InWorld, bool bFuzzy = false)
@@ -193,7 +134,7 @@ struct FActorTakeRecorderDropHandler : ITakeRecorderDropHandler
 	virtual void HandleOperation(TSharedPtr<FDragDropOperation> InOperation, UTakeRecorderSources* Sources) override
 	{
 		TArray<AActor*> ActorsToAdd = GetValidDropActors(InOperation, Sources);
-		AddActorSources(Sources, ActorsToAdd);
+		TakeRecorderSourceHelpers::AddActorSources(Sources, ActorsToAdd);
 	}
 
 	virtual bool CanHandleOperation(TSharedPtr<FDragDropOperation> InOperation, UTakeRecorderSources* Sources) override
@@ -534,7 +475,7 @@ public:
 				FText(),
 				ActorIcon,
 				FExecuteAction::CreateLambda([Sources, SelectedActors]{
-					AddActorSources(Sources, SelectedActors);
+					TakeRecorderSourceHelpers::AddActorSources(Sources, SelectedActors);
 				})
 			);
 		}
@@ -569,7 +510,7 @@ public:
 						FOnActorPicked::CreateLambda([Sources](AActor* Actor){
 							// Create a new binding for this actor
 							FSlateApplication::Get().DismissAllMenus();
-							AddActorSources(Sources, MakeArrayView(&Actor, 1));
+							TakeRecorderSourceHelpers::AddActorSources(Sources, MakeArrayView(&Actor, 1));
 						})
 					)
 				];
diff --git a/Engine/Plugins/VirtualProduction/Takes/Source/TakeRecorderSources/Public/TakeRecorderSourceHelpers.h b/Engine/Plugins/VirtualProduction/Takes/Source/TakeRecorderSources/Public/TakeRecorderSourceHelpers.h
new file mode 100644
--- /dev/null
+++ b/Engine/Plugins/VirtualProduction/Takes/Source/TakeRecorderSources/Public/TakeRecorderSourceHelpers.h
@@ -0,0 +1,32 @@
+// Copyright Epic Games, Inc. All Rights Reserved.
+
+#pragma once
+
+#include "CoreMinimal.h"
+#include "Containers/ArrayView.h"
+#include "TakeRecorderSources.h"
+#include "GameFramework/Actor.h"
+
+class AActor;
+
+namespace TakeRecorderSourceHelpers
+{
+	/**
+	 * Adds a number of actors to the set of sources to record from.
+	 * 
+	 * @param TakeRecorderSources The list of sources used for the current take. 
+	 * @param ActorsToRecord The list of Actors that should be added to Sources. Note that this can include ALevelSequenceActors.
+	 * @param bReduceKeys Enable/disable key reduction on all the sources registered
+	 * @param bShowProgress Enable/disable the dialog box showing progress for the potentially slow parts of finalizing the take
+	 */
+	TAKERECORDERSOURCES_API void AddActorSources(
+		UTakeRecorderSources*     TakeRecorderSources,
+		TArrayView<AActor* const> ActorsToRecord,
+		bool                      bReduceKeys = true,
+		bool                      bShowProgress = true);
+
+	/**
+	 * Removes all sources from a list of sources to record from.
+	 */
+	TAKERECORDERSOURCES_API void RemoveAllActorSources(UTakeRecorderSources* Sources);
+};
diff --git a/Engine/Plugins/VirtualProduction/Takes/Source/TakeRecorderSources/TakeRecorderSources.Build.cs b/Engine/Plugins/VirtualProduction/Takes/Source/TakeRecorderSources/TakeRecorderSources.Build.cs
--- a/Engine/Plugins/VirtualProduction/Takes/Source/TakeRecorderSources/TakeRecorderSources.Build.cs
+++ b/Engine/Plugins/VirtualProduction/Takes/Source/TakeRecorderSources/TakeRecorderSources.Build.cs
@@ -42,6 +42,7 @@ public class TakeRecorderSources : ModuleRules
 		PrivateIncludePaths.AddRange(
 			new string[] {
 				"TakeRecorderSources/Private",
+				"TakeRecorderSources/Public",
 			}
 		);
 	}
diff --git a/Engine/Source/Editor/Sequencer/Private/Sequencer.cpp b/Engine/Source/Editor/Sequencer/Private/Sequencer.cpp
--- a/Engine/Source/Editor/Sequencer/Private/Sequencer.cpp
+++ b/Engine/Source/Editor/Sequencer/Private/Sequencer.cpp
@@ -1401,6 +1401,10 @@ void FSequencer::RerunConstructionScripts()
 	FMovieSceneRootEvaluationTemplateInstance& RootTemplate = GetEvaluationTemplate();
 		
 	UMovieSceneSequence* Sequence = RootTemplate.GetSequence(MovieSceneSequenceID::Root);
+	if (!Sequence)
+	{
+		return;
+	}
 
 	TArray < TPair<FMovieSceneSequenceID, FGuid> > BoundGuids;
 
@@ -6624,7 +6628,7 @@ void FSequencer::SynchronizeSequencerSelectionWithExternalSelection()
 		return;
 	}
 
-	if (!Sequence->GetMovieScene())
+	if (!Sequence || !Sequence->GetMovieScene())
 	{
 		return;
 	}
