Take Recorder: Show the engine time dilation in the cockpit.

Add a warning if the time dilation is 0

diff --git a/Engine/Plugins/VirtualProduction/Takes/Source/TakeRecorder/Private/Recorder/TakeRecorder.cpp b/Engine/Plugins/VirtualProduction/Takes/Source/TakeRecorder/Private/Recorder/TakeRecorder.cpp
--- a/Engine/Plugins/VirtualProduction/Takes/Source/TakeRecorder/Private/Recorder/TakeRecorder.cpp
+++ b/Engine/Plugins/VirtualProduction/Takes/Source/TakeRecorder/Private/Recorder/TakeRecorder.cpp
@@ -930,7 +930,14 @@ void UTakeRecorder::PreRecord()
 	{
 		const float ExistingMatineeTimeDilation = WorldSettings->MatineeTimeDilation;
 
-		if (Parameters.User.EngineTimeDilation != ExistingMatineeTimeDilation)
+		const bool bInvalidTimeDilation = Parameters.User.EngineTimeDilation == 0.f;
+
+		if (bInvalidTimeDilation)
+		{
+			UE_LOG(LogTakesCore, Warning, TEXT("Time dilation cannot be 0. Ignoring time dilation for this recording."));
+		}
+
+		if (Parameters.User.EngineTimeDilation != ExistingMatineeTimeDilation && !bInvalidTimeDilation)
 		{
 			WorldSettings->MatineeTimeDilation = Parameters.User.EngineTimeDilation;
 
diff --git a/Engine/Plugins/VirtualProduction/Takes/Source/TakeRecorder/Private/Widgets/STakeRecorderCockpit.cpp b/Engine/Plugins/VirtualProduction/Takes/Source/TakeRecorder/Private/Widgets/STakeRecorderCockpit.cpp
--- a/Engine/Plugins/VirtualProduction/Takes/Source/TakeRecorder/Private/Widgets/STakeRecorderCockpit.cpp
+++ b/Engine/Plugins/VirtualProduction/Takes/Source/TakeRecorder/Private/Widgets/STakeRecorderCockpit.cpp
@@ -458,14 +458,33 @@ void STakeRecorderCockpit::Construct(const FArguments& InArgs)
 				.Padding(8, 0, 8, 8)
 				.AutoHeight()
 				[
-					SNew(SEditableTextBox)
-					.IsEnabled(this, &STakeRecorderCockpit::EditingMetaData)
-					.Style(FTakeRecorderStyle::Get(), "TakeRecorder.EditableTextBox")
-					.Font(FTakeRecorderStyle::Get().GetFontStyle("TakeRecorder.Cockpit.SmallText"))
-					.SelectAllTextWhenFocused(true)
-					.HintText(LOCTEXT("EnterSlateDescription_Hint", "<description>"))
-					.Text(this, &STakeRecorderCockpit::GetUserDescriptionText)
-					.OnTextCommitted(this, &STakeRecorderCockpit::SetUserDescriptionText)
+					SNew(SHorizontalBox)
+
+					+ SHorizontalBox::Slot()
+					[
+						SNew(SEditableTextBox)
+						.IsEnabled(this, &STakeRecorderCockpit::EditingMetaData)
+						.Style(FTakeRecorderStyle::Get(), "TakeRecorder.EditableTextBox")
+						.Font(FTakeRecorderStyle::Get().GetFontStyle("TakeRecorder.Cockpit.SmallText"))
+						.SelectAllTextWhenFocused(true)
+						.HintText(LOCTEXT("EnterSlateDescription_Hint", "<description>"))
+						.Text(this, &STakeRecorderCockpit::GetUserDescriptionText)
+						.OnTextCommitted(this, &STakeRecorderCockpit::SetUserDescriptionText)
+					]
+
+					+ SHorizontalBox::Slot()
+					.AutoWidth()
+					[
+						SNew(SSpinBox<float>)
+						.ToolTipText(LOCTEXT("EngineTimeDilation", "Recording speed"))
+						.Style(&FAppStyle::GetWidgetStyle<FSpinBoxStyle>("Sequencer.HyperlinkSpinBox"))
+						.OnValueChanged(this, &STakeRecorderCockpit::SetEngineTimeDilation)
+						.OnValueCommitted_Lambda([=](float InEngineTimeDilation, ETextCommit::Type) { SetEngineTimeDilation(InEngineTimeDilation); })
+						.MinValue(TOptional<float>())
+						.MaxValue(TOptional<float>())
+						.Value(this, &STakeRecorderCockpit::GetEngineTimeDilation)
+						.Delta(0.5f)
+					]
 				]
 			]
 		]
@@ -907,6 +926,16 @@ void STakeRecorderCockpit::OnEndSetTakeNumber(int32 InFinalValue)
 	TransactionIndex = INDEX_NONE;
 }
 
+float STakeRecorderCockpit::GetEngineTimeDilation() const
+{
+	return GetDefault<UTakeRecorderUserSettings>()->Settings.EngineTimeDilation;
+}
+
+void STakeRecorderCockpit::SetEngineTimeDilation(float InEngineTimeDilation)
+{
+	GetMutableDefault<UTakeRecorderUserSettings>()->Settings.EngineTimeDilation = InEngineTimeDilation;
+}
+
 FReply STakeRecorderCockpit::OnAddMarkedFrame()
 {
 	if (UTakeRecorderBlueprintLibrary::IsRecording())
diff --git a/Engine/Plugins/VirtualProduction/Takes/Source/TakeRecorder/Private/Widgets/STakeRecorderCockpit.h b/Engine/Plugins/VirtualProduction/Takes/Source/TakeRecorder/Private/Widgets/STakeRecorderCockpit.h
--- a/Engine/Plugins/VirtualProduction/Takes/Source/TakeRecorder/Private/Widgets/STakeRecorderCockpit.h
+++ b/Engine/Plugins/VirtualProduction/Takes/Source/TakeRecorder/Private/Widgets/STakeRecorderCockpit.h
@@ -119,6 +119,10 @@ private:
 
 	void OnEndSetTakeNumber(int32 InFinalValue);
 
+	float GetEngineTimeDilation() const;
+
+	void SetEngineTimeDilation(float InEngineTimeDilation);
+
 	FReply OnAddMarkedFrame();
 
 	void OnToggleRecording(ECheckBoxState);
