Take Recorder: Added recording speed label

diff --git a/Engine/Plugins/VirtualProduction/Takes/Source/TakeRecorder/Private/Widgets/STakeRecorderCockpit.cpp b/Engine/Plugins/VirtualProduction/Takes/Source/TakeRecorder/Private/Widgets/STakeRecorderCockpit.cpp
--- a/Engine/Plugins/VirtualProduction/Takes/Source/TakeRecorder/Private/Widgets/STakeRecorderCockpit.cpp
+++ b/Engine/Plugins/VirtualProduction/Takes/Source/TakeRecorder/Private/Widgets/STakeRecorderCockpit.cpp
@@ -485,6 +485,18 @@ void STakeRecorderCockpit::Construct(const FArguments& InArgs)
 						.Value(this, &STakeRecorderCockpit::GetEngineTimeDilation)
 						.Delta(0.5f)
 					]
+
+					+ SHorizontalBox::Slot()
+					.Padding(2, 0, 0, 2)
+					.VAlign(VAlign_Bottom)
+					.AutoWidth()
+					[
+						SNew(STextBlock)
+						.ColorAndOpacity(FSlateColor::UseSubduedForeground())
+						.Font(FTakeRecorderStyle::Get().GetFontStyle("TakeRecorder.Cockpit.SmallText"))
+						.TextStyle(FTakeRecorderStyle::Get(), "TakeRecorder.TextBox")
+						.Text(LOCTEXT("EngineTimeDilationLabel", "x"))
+					]
 				]
 			]
 		]
