Take Recorder: Move UpdateTakeError() from Tick() to when the slate has changed. The asset registry can be slow, so it should be sufficient to just do the update when the slate has changed.

diff --git a/Engine/Plugins/VirtualProduction/Takes/Source/TakeRecorder/Private/Widgets/STakeRecorderCockpit.cpp b/Engine/Plugins/VirtualProduction/Takes/Source/TakeRecorder/Private/Widgets/STakeRecorderCockpit.cpp
--- a/Engine/Plugins/VirtualProduction/Takes/Source/TakeRecorder/Private/Widgets/STakeRecorderCockpit.cpp
+++ b/Engine/Plugins/VirtualProduction/Takes/Source/TakeRecorder/Private/Widgets/STakeRecorderCockpit.cpp
@@ -135,6 +135,9 @@ void STakeRecorderCockpit::Construct(const FArguments& InArgs)
 	TakeMetaData = nullptr;
 	TransientTakeMetaData = nullptr;
 
+	CachedTakeSlate.Empty();
+	CachedTakeNumber = -1;
+
 	LevelSequenceAttribute = InArgs._LevelSequence;
 
 	TakeRecorderModeAttribute = InArgs._TakeRecorderMode;
@@ -714,7 +717,16 @@ void STakeRecorderCockpit::CacheMetaData()
 	if (NewMetaDataThisTick != TakeMetaData)
 	{
 		TakeMetaData = NewMetaDataThisTick;
-		// Forcibly update any UI?
+	}
+
+	if (TakeMetaData->GetSlate() != CachedTakeSlate || TakeMetaData->GetTakeNumber() != CachedTakeNumber)
+	{
+		CachedTakeNumber = TakeMetaData->GetTakeNumber();
+		CachedTakeSlate = TakeMetaData->GetSlate();
+
+		// Previously, the take error would be updated in Tick(), but the asset registry can be slow, 
+		// so it should be sufficient to update it only when the slate changes.
+		UpdateTakeError();
 	}
 
 	//Set MovieScene Display Rate to the Preset Frame Rate.
@@ -736,7 +748,6 @@ void STakeRecorderCockpit::Tick(const FGeometry& AllottedGeometry, const double
 void STakeRecorderCockpit::Refresh()
 {
 	CacheMetaData();
-	UpdateTakeError();
 	UpdateRecordError();
 }
 
diff --git a/Engine/Plugins/VirtualProduction/Takes/Source/TakeRecorder/Private/Widgets/STakeRecorderCockpit.h b/Engine/Plugins/VirtualProduction/Takes/Source/TakeRecorder/Private/Widgets/STakeRecorderCockpit.h
--- a/Engine/Plugins/VirtualProduction/Takes/Source/TakeRecorder/Private/Widgets/STakeRecorderCockpit.h
+++ b/Engine/Plugins/VirtualProduction/Takes/Source/TakeRecorder/Private/Widgets/STakeRecorderCockpit.h
@@ -181,4 +181,8 @@ private:
 	TSharedPtr<FUICommandList> CommandList;
 
 	TSharedPtr<FDigitsTypeInterface> DigitsTypeInterface;
+
+	// Cached take numbers and slate used to UpdateTakeError() only when necesary
+	int32 CachedTakeNumber;
+	FString CachedTakeSlate;
 };
