Take Recorder: Don't allow delete to trickle through locked sources and end up deleting selected actors for example

diff --git a/Engine/Plugins/VirtualProduction/Takes/Source/TakeRecorder/Private/Widgets/STakeRecorderSources.cpp b/Engine/Plugins/VirtualProduction/Takes/Source/TakeRecorder/Private/Widgets/STakeRecorderSources.cpp
--- a/Engine/Plugins/VirtualProduction/Takes/Source/TakeRecorder/Private/Widgets/STakeRecorderSources.cpp
+++ b/Engine/Plugins/VirtualProduction/Takes/Source/TakeRecorder/Private/Widgets/STakeRecorderSources.cpp
@@ -393,8 +393,7 @@ void STakeRecorderSources::Construct(const FArguments& InArgs)
 	CommandList = MakeShared<FUICommandList>();
 	CommandList->MapAction(
 		FGenericCommands::Get().Delete,
-		FExecuteAction::CreateSP(this, &STakeRecorderSources::OnDeleteSelected),
-		FCanExecuteAction::CreateSP(this, &STakeRecorderSources::CanDeleteSelected)
+		FExecuteAction::CreateSP(this, &STakeRecorderSources::OnDeleteSelected)
 	);
 
 	ChildSlot
@@ -645,6 +644,12 @@ bool STakeRecorderSources::CanDragDropTarget(TSharedPtr<FDragDropOperation> InOp
 
 void STakeRecorderSources::OnDeleteSelected()
 {
+	// Still capture the delete so that it doesn't trickle through to delete selected actors for example
+	if (IsLocked())
+	{
+		return;
+	}
+
 	UTakeRecorderSources* Sources = WeakSources.Get();
 	if (Sources)
 	{
@@ -660,11 +665,6 @@ void STakeRecorderSources::OnDeleteSelected()
 	}
 }
 
-bool STakeRecorderSources::CanDeleteSelected() const
-{
-	return !IsLocked();
-}
-
 bool STakeRecorderSources::IsLocked() const
 {
 	UTakeRecorderSources* Sources = WeakSources.Get();
diff --git a/Engine/Plugins/VirtualProduction/Takes/Source/TakeRecorder/Private/Widgets/STakeRecorderSources.h b/Engine/Plugins/VirtualProduction/Takes/Source/TakeRecorder/Private/Widgets/STakeRecorderSources.h
--- a/Engine/Plugins/VirtualProduction/Takes/Source/TakeRecorder/Private/Widgets/STakeRecorderSources.h
+++ b/Engine/Plugins/VirtualProduction/Takes/Source/TakeRecorder/Private/Widgets/STakeRecorderSources.h
@@ -58,8 +58,6 @@ private:
 
 	void OnDeleteSelected();
 
-	bool CanDeleteSelected() const;
-
 	bool IsLocked() const;
 
 private:
