Take Recorder: Use current timecode when recording timecode, not the sequencer time transformed

diff --git a/Engine/Plugins/VirtualProduction/Takes/Source/TakeRecorderSources/Private/TakeRecorderActorSource.cpp b/Engine/Plugins/VirtualProduction/Takes/Source/TakeRecorderSources/Private/TakeRecorderActorSource.cpp
--- a/Engine/Plugins/VirtualProduction/Takes/Source/TakeRecorderSources/Private/TakeRecorderActorSource.cpp
+++ b/Engine/Plugins/VirtualProduction/Takes/Source/TakeRecorderSources/Private/TakeRecorderActorSource.cpp
@@ -685,7 +685,7 @@ void UTakeRecorderActorSource::ProcessRecordedTimes(ULevelSequence* InSequence)
 		TArray<FMovieSceneFloatValue> SubFrames;
 		TArray<FFrameNumber> Times;
 
-		const TArray<TPair<FQualifiedFrameTime, FTimecode>>& RecordedTimes = UTakeRecorderSources::RecordedTimes;
+		const TArray<TPair<FQualifiedFrameTime, FQualifiedFrameTime>>& RecordedTimes = UTakeRecorderSources::RecordedTimes;
 
 		Hours.Reserve(RecordedTimes.Num());
 		Minutes.Reserve(RecordedTimes.Num());
@@ -697,7 +697,7 @@ void UTakeRecorderActorSource::ProcessRecordedTimes(ULevelSequence* InSequence)
 		FFrameRate TickResolution = MovieScene->GetTickResolution();
 		FFrameRate DisplayRate = MovieScene->GetDisplayRate();
 
-		for (const TPair<FQualifiedFrameTime, FTimecode>& RecordedTimePair : RecordedTimes)
+		for (const TPair<FQualifiedFrameTime, FQualifiedFrameTime>& RecordedTimePair : RecordedTimes)
 		{
 			FFrameNumber FrameNumber = RecordedTimePair.Key.Time.FrameNumber;
 			if (!FrameRange.GetValue().Contains(FrameNumber))
@@ -707,7 +707,7 @@ void UTakeRecorderActorSource::ProcessRecordedTimes(ULevelSequence* InSequence)
 
 			FFrameTime FrameTime = FFrameRate::TransformTime(RecordedTimePair.Key.Time, TickResolution, DisplayRate);
 
-			FTimecode Timecode = RecordedTimePair.Value;
+			FTimecode Timecode = RecordedTimePair.Value.ToTimecode();
 		
 			Hours.Add(Timecode.Hours);
 			Minutes.Add(Timecode.Minutes);
@@ -715,7 +715,7 @@ void UTakeRecorderActorSource::ProcessRecordedTimes(ULevelSequence* InSequence)
 			Frames.Add(Timecode.Frames);
 
 			FMovieSceneFloatValue SubFrame;
-			SubFrame.Value = FrameTime.GetSubFrame();
+			SubFrame.Value = RecordedTimePair.Value.Time.GetSubFrame();
 			SubFrame.InterpMode = ERichCurveInterpMode::RCIM_Linear;
 			SubFrames.Add(SubFrame);
 
diff --git a/Engine/Plugins/VirtualProduction/Takes/Source/TakesCore/Private/TakeRecorderSources.cpp b/Engine/Plugins/VirtualProduction/Takes/Source/TakesCore/Private/TakeRecorderSources.cpp
--- a/Engine/Plugins/VirtualProduction/Takes/Source/TakesCore/Private/TakeRecorderSources.cpp
+++ b/Engine/Plugins/VirtualProduction/Takes/Source/TakesCore/Private/TakeRecorderSources.cpp
@@ -22,7 +22,7 @@
 
 DEFINE_LOG_CATEGORY(SubSequenceSerialization);
 
-TArray<TPair<FQualifiedFrameTime, FTimecode> > UTakeRecorderSources::RecordedTimes;
+TArray<TPair<FQualifiedFrameTime, FQualifiedFrameTime> > UTakeRecorderSources::RecordedTimes;
 
 UTakeRecorderSources::UTakeRecorderSources(const FObjectInitializer& ObjInit)
 	: Super(ObjInit)
@@ -365,8 +365,9 @@
 		}
 	}
 
+	if (FApp::GetCurrentFrameTime().IsSet())
 	{
-		RecordedTimes.Add(TPair<FQualifiedFrameTime, FTimecode>(CurrentFrameTime, Timecode));
+		RecordedTimes.Add(TPair<FQualifiedFrameTime, FQualifiedFrameTime>(CurrentFrameTime, FApp::GetCurrentFrameTime().GetValue() ));
 	}
 
 	// If we're recording into sub-sections we want to update their range every frame so they appear to
diff --git a/Engine/Plugins/VirtualProduction/Takes/Source/TakesCore/Public/TakeRecorderSources.h b/Engine/Plugins/VirtualProduction/Takes/Source/TakesCore/Public/TakeRecorderSources.h
--- a/Engine/Plugins/VirtualProduction/Takes/Source/TakesCore/Public/TakeRecorderSources.h
+++ b/Engine/Plugins/VirtualProduction/Takes/Source/TakesCore/Public/TakeRecorderSources.h
@@ -171,7 +171,7 @@ public:
 	/**
 	 * Array of pairs - key time and the corresponding timecode
 	 */
-	static TArray<TPair<FQualifiedFrameTime, FTimecode> > RecordedTimes;
+	static TArray<TPair<FQualifiedFrameTime, FQualifiedFrameTime> > RecordedTimes;
 
 	FQualifiedFrameTime GetCachedFrameTime() const { return CachedFrameTime; }
 
