Takes: Move TakeRecorder dynamic delegates to global namespace or else they're not usable via Python.

diff --git a/Engine/Plugins/VirtualProduction/Takes/Source/TakeRecorder/Private/Recorder/TakeRecorderBlueprintLibrary.cpp b/Engine/Plugins/VirtualProduction/Takes/Source/TakeRecorder/Private/Recorder/TakeRecorderBlueprintLibrary.cpp
--- a/Engine/Plugins/VirtualProduction/Takes/Source/TakeRecorder/Private/Recorder/TakeRecorderBlueprintLibrary.cpp
+++ b/Engine/Plugins/VirtualProduction/Takes/Source/TakeRecorder/Private/Recorder/TakeRecorderBlueprintLibrary.cpp
@@ -14,13 +14,13 @@
 namespace
 {
 	static UTakeRecorderPanel* CurrentTakeRecorderPanel;
-	static UTakeRecorderBlueprintLibrary::FOnTakeRecorderPanelChanged TakeRecorderPanelChanged;
-	static UTakeRecorderBlueprintLibrary::FOnTakeRecorderPreInitialize TakeRecorderPreInitialize;
-	static UTakeRecorderBlueprintLibrary::FOnTakeRecorderStarted TakeRecorderStarted;
-	static UTakeRecorderBlueprintLibrary::FOnTakeRecorderStopped TakeRecorderStopped;
-	static UTakeRecorderBlueprintLibrary::FOnTakeRecorderFinished TakeRecorderFinished;
-	static UTakeRecorderBlueprintLibrary::FOnTakeRecorderCancelled TakeRecorderCancelled;
-	static UTakeRecorderBlueprintLibrary::FOnTakeRecorderMarkedFrameAdded TakeRecorderMarkedFrameAdded;
+	static FOnTakeRecorderPanelChanged TakeRecorderPanelChanged;
+	static FOnTakeRecorderPreInitialize TakeRecorderPreInitialize;
+	static FOnTakeRecorderStarted TakeRecorderStarted;
+	static FOnTakeRecorderStopped TakeRecorderStopped;
+	static FOnTakeRecorderFinished TakeRecorderFinished;
+	static FOnTakeRecorderCancelled TakeRecorderCancelled;
+	static FOnTakeRecorderMarkedFrameAdded TakeRecorderMarkedFrameAdded;
 }
 
 
