Update sequencer view-models on take recorder tick

diff --git a/Engine/Plugins/VirtualProduction/Takes/Source/TakeRecorder/Private/Recorder/TakeRecorder.cpp b/Engine/Plugins/VirtualProduction/Takes/Source/TakeRecorder/Private/Recorder/TakeRecorder.cpp
--- a/Engine/Plugins/VirtualProduction/Takes/Source/TakeRecorder/Private/Recorder/TakeRecorder.cpp
+++ b/Engine/Plugins/VirtualProduction/Takes/Source/TakeRecorder/Private/Recorder/TakeRecorder.cpp
@@ -817,6 +817,8 @@ FQualifiedFrameTime UTakeRecorder::GetRecordTime() const
 
 void UTakeRecorder::InternalTick(float DeltaTime)
 {
+	UE::MovieScene::FScopedSignedObjectModifyDefer FlushOnTick(true);
+
 	TSharedPtr<ISequencer> Sequencer = WeakSequencer.Pin();
 	FQualifiedFrameTime RecordTime = GetRecordTime();
 		
