Take Recorder: Spaces in Actor names can cause invalid anim sequences to get recorded. Replace " " with "_". Tried the various FPath:: functions but non of them dealt with white space.
jira #UE-15669

diff --git a/Engine/Plugins/VirtualProduction/Takes/Source/TakesCore/Public/TakesUtils.h b/Engine/Plugins/VirtualProduction/Takes/Source/TakesCore/Public/TakesUtils.h
--- a/Engine/Plugins/VirtualProduction/Takes/Source/TakesCore/Public/TakesUtils.h
+++ b/Engine/Plugins/VirtualProduction/Takes/Source/TakesCore/Public/TakesUtils.h
@@ -113,6 +113,9 @@ namespace TakesUtils
 		FString AssetPath = BaseAssetPath;
 		FString AssetName = BaseAssetName;
 		AssetName = AssetName.Replace(TEXT("."), TEXT("_"));
+		AssetName = AssetName.Replace(TEXT(" "), TEXT("_"));
+
+		AssetName = FPaths::MakeValidFileName(AssetName);
 
 		AssetPath /= AssetName;
 		AssetPath += Dot + AssetName;
