Take Recorder: Reset transform for recorded spawnable actors when their skeletal animation is recorded in world space but a transform track is not being recorded.

This is a very specific case with the following parameters:
1. Skeletal mesh actor is recorded as a spawnable.
2. Skeletal mesh actor is a child of another actor which is not being recorded (this results in the bones being recorded in world space)
3. Transform track is set to not record.

When this occurs, the transform of the spawnable is doubled up because the root bone is in world space and the spawnable template is also in world space. The fix here is to reset the spawnable template's transform to identity.

diff --git a/Engine/Plugins/VirtualProduction/Takes/Source/TakeRecorderSources/Private/TakeRecorderActorSource.cpp b/Engine/Plugins/VirtualProduction/Takes/Source/TakeRecorderSources/Private/TakeRecorderActorSource.cpp
--- a/Engine/Plugins/VirtualProduction/Takes/Source/TakeRecorderSources/Private/TakeRecorderActorSource.cpp
+++ b/Engine/Plugins/VirtualProduction/Takes/Source/TakeRecorderSources/Private/TakeRecorderActorSource.cpp
@@ -929,6 +929,37 @@ void UTakeRecorderActorSource::PostProcessTrackRecorders(ULevelSequence* InSeque
 			}
 		}
 	}
+
+	// Reset transform for recorded spawnable actors when their skeletal animation is recorded in world space 
+	// but a transform track is not being recorded.
+	// 
+	// This is a very specific case with the following parameters:
+	// 1. Skeletal mesh actor is recorded as a spawnable.
+	// 2. Skeletal mesh actor is a child of another actor which is not being recorded (this results in the bones being recorded in world space)
+	// 3. Transform track is set to not record.
+	//
+	// When this occurs, the transform of the spawnable is doubled up because the root bone is in world space
+	// and the spawnable template is also in world space.The fix here is to reset the spawnable template's 
+	// transform to identity.	
+	if (FirstAnimationRecorder && !RootTransformRecorder)
+	{
+		if (CachedObjectTemplate.IsValid())
+		{
+			AActor* ActorToRecord = Target.Get();
+			if (ActorToRecord)
+			{
+				USceneComponent* RootComponent = ActorToRecord->GetRootComponent();
+				USceneComponent* AttachParent = RootComponent ? RootComponent->GetAttachParent() : nullptr;
+				if (AttachParent && !IsOtherActorBeingRecorded(AttachParent->GetOwner()))
+				{
+					// The object template is marked as bComponentToWorldUpdated=true while the ComponentToWorld doesn't 
+					// match the relative location and rotation. So, calling SetRelativeTransform() doesn't work. 
+					// Set it directly here.
+					CachedObjectTemplate->GetRootComponent()->SetRelativeTransform_Direct(FTransform::Identity);
+				}
+			}
+		}
+	}
 }
 
 bool UTakeRecorderActorSource::EnsureObjectTemplateHasComponent(UActorComponent* InComponent, UActorComponent*& OutComponent)
