NonUnity BuildFix for TakesUtils.h, missing .h file
#trvial

diff --git a/Engine/Plugins/VirtualProduction/Takes/Source/TakesCore/Public/TakesUtils.h b/Engine/Plugins/VirtualProduction/Takes/Source/TakesCore/Public/TakesUtils.h
--- a/Engine/Plugins/VirtualProduction/Takes/Source/TakesCore/Public/TakesUtils.h
+++ b/Engine/Plugins/VirtualProduction/Takes/Source/TakesCore/Public/TakesUtils.h
@@ -9,6 +9,7 @@
 #include "AssetRegistry/AssetRegistryModule.h"
 #include "UObject/UObjectGlobals.h"
 #include "TakesCoreLog.h"
+#include "Misc/Paths.h"
 
 class UWorld;
 class ULevelSequence;
