Take Recorder: Fix review last take when can review is not bound

diff --git a/Engine/Plugins/VirtualProduction/Takes/Source/TakeRecorder/Private/Widgets/STakeRecorderPanel.cpp b/Engine/Plugins/VirtualProduction/Takes/Source/TakeRecorder/Private/Widgets/STakeRecorderPanel.cpp
--- a/Engine/Plugins/VirtualProduction/Takes/Source/TakeRecorder/Private/Widgets/STakeRecorderPanel.cpp
+++ b/Engine/Plugins/VirtualProduction/Takes/Source/TakeRecorder/Private/Widgets/STakeRecorderPanel.cpp
@@ -570,7 +570,12 @@ bool STakeRecorderPanel::CanReviewLastLevelSequence() const
 	}
 	ITakeRecorderModule& TakeRecorderModule = FModuleManager::Get().LoadModuleChecked<ITakeRecorderModule>("TakeRecorder");
 	FCanReviewLastRecordedLevelSequence& CanReview = TakeRecorderModule.GetCanReviewLastRecordedLevelSequenceDelegate();
-	return CanReview.IsBound() && CanReview.Execute();
+	if (CanReview.IsBound())
+	{
+		return CanReview.Execute();
+	}
+	
+	return true;
 }
 
 ETakeRecorderMode STakeRecorderPanel::GetTakeRecorderMode() const
