Take Recorder: Disregard adding a parent source for recording if it's been added as a source and disabled explicitly by the user.

diff --git a/Engine/Plugins/VirtualProduction/Takes/Source/TakeRecorderSources/Private/TakeRecorderActorSource.cpp b/Engine/Plugins/VirtualProduction/Takes/Source/TakeRecorderSources/Private/TakeRecorderActorSource.cpp
--- a/Engine/Plugins/VirtualProduction/Takes/Source/TakeRecorderSources/Private/TakeRecorderActorSource.cpp
+++ b/Engine/Plugins/VirtualProduction/Takes/Source/TakeRecorderSources/Private/TakeRecorderActorSource.cpp
@@ -1600,12 +1600,40 @@ void UTakeRecorderActorSource::CreateNewActorSourceForReferencedActors()
 	UTakeRecorderSources* SourcesList = GetTypedOuter<UTakeRecorderSources>();
 	TArray<UTakeRecorderSource*> NewSources;
 
+	TSet<AActor*> ActorsWithEnabledSources;
+	TSet<AActor*> ActorsWithDisabledSources;
+	for (UTakeRecorderSource* Source : SourcesList->GetSources())
+	{
+		if (UTakeRecorderActorSource* ActorSource = Cast<UTakeRecorderActorSource>(Source))
+		{
+			AActor* TargetActor = ActorSource->Target.Get();
+			if (TargetActor)
+			{
+				if (ActorSource->bEnabled)
+				{
+					ActorsWithEnabledSources.Add(TargetActor);
+				}
+				else
+				{
+					ActorsWithDisabledSources.Add(TargetActor);
+				}
+			}
+		}
+	}
+
 	for (AActor* Actor : NewReferencedActors)
 	{
-		if (IsOtherActorBeingRecorded(Actor))
+		// Don't create a recording for this actor if there is an existing source for this actor. Another source may have added it
+		// or the user may have added it by hand and adjusted settings.
+		if (ActorsWithEnabledSources.Contains(Actor))
+		{
+			continue;
+		}
+
+		// Also, don't create a recording if the actor has a disabled source
+		if (ActorsWithDisabledSources.Contains(Actor))
 		{
-			// Don't create a recording for this actor if they're already recording it. Another source may have added it
-			// or the user may have added it by hand and adjusted settings.
+			UE_LOG(LogTakesCore, Warning, TEXT("Disregarding automatically adding %s as a recording source because it has been explicitly disabled."), *Actor->GetName());
 			continue;
 		}
 
