Take Recorder: Fix for regression on recording float properties (allow recording of double properties)

diff --git a/Engine/Plugins/VirtualProduction/Takes/Source/TakeTrackRecorders/Private/MovieScenePropertyTrackRecorder.cpp b/Engine/Plugins/VirtualProduction/Takes/Source/TakeTrackRecorders/Private/MovieScenePropertyTrackRecorder.cpp
--- a/Engine/Plugins/VirtualProduction/Takes/Source/TakeTrackRecorders/Private/MovieScenePropertyTrackRecorder.cpp
+++ b/Engine/Plugins/VirtualProduction/Takes/Source/TakeTrackRecorders/Private/MovieScenePropertyTrackRecorder.cpp
@@ -30,6 +30,10 @@ bool FMovieScenePropertyTrackRecorderFactory::CanRecordProperty(UObject* InObjec
 	{
 		return true;
 	}
+	else if (InPropertyToRecord->IsA<FDoubleProperty>())
+	{
+		return true;
+	}
 	else if (FStructProperty* StructProperty = CastField<FStructProperty>(InPropertyToRecord))
 	{
 		if (StructProperty->Struct->GetFName() == NAME_Vector)
@@ -130,10 +134,14 @@ void UMovieScenePropertyTrackRecorder::CreateTrackImpl()
 			PropertyRecorder = MakeShareable(new FMovieSceneTrackPropertyRecorder<FString>(Binding));
 		}
 		else if (Property->IsA<FFloatProperty>())
- 		{
+		{
 			PropertyRecorder = MakeShareable(new FMovieSceneTrackPropertyRecorder<float>(Binding));
- 		}
- 		else if (FStructProperty* StructProperty = CastField<FStructProperty>(Property))
+		}
+		else if (Property->IsA<FDoubleProperty>())
+		{
+			PropertyRecorder = MakeShareable(new FMovieSceneTrackPropertyRecorder<double>(Binding));
+		}
+		else if (FStructProperty* StructProperty = CastField<FStructProperty>(Property))
  		{
 			// LWC_TODO: vector types
 			if (StructProperty->Struct->GetFName() == NAME_Vector3f)
