FNames containing asset paths are deprecated. FSoftObjectPath or FTopLevelAssetPath should be used instead.
Fixups for VirtualProduction.

diff --git a/Engine/Plugins/VirtualProduction/CameraCalibration/Source/CameraCalibrationEditor/Private/CameraCalibrationEditorModule.cpp b/Engine/Plugins/VirtualProduction/CameraCalibration/Source/CameraCalibrationEditor/Private/CameraCalibrationEditorModule.cpp
--- a/Engine/Plugins/VirtualProduction/CameraCalibration/Source/CameraCalibrationEditor/Private/CameraCalibrationEditorModule.cpp
+++ b/Engine/Plugins/VirtualProduction/CameraCalibration/Source/CameraCalibrationEditor/Private/CameraCalibrationEditorModule.cpp
@@ -169,7 +169,7 @@ void FCameraCalibrationEditorModule::RegisterOverlayMaterials()
 	{
 		// Register all overlay materials defined in this module
 		UCameraCalibrationSubsystem* SubSystem = GEngine->GetEngineSubsystem<UCameraCalibrationSubsystem>();
-		SubSystem->RegisterOverlayMaterial(TEXT("Crosshair"), TEXT("/CameraCalibration/Materials/M_Crosshair.M_Crosshair"));
+		SubSystem->RegisterOverlayMaterial(TEXT("Crosshair"), FSoftObjectPath(TEXT("/CameraCalibration/Materials/M_Crosshair.M_Crosshair")));
 	};
 
 	if (FApp::CanEverRender())
diff --git a/Engine/Plugins/VirtualProduction/CameraCalibrationCore/Source/CameraCalibrationCore/Private/CameraCalibrationSubsystem.cpp b/Engine/Plugins/VirtualProduction/CameraCalibrationCore/Source/CameraCalibrationCore/Private/CameraCalibrationSubsystem.cpp
--- a/Engine/Plugins/VirtualProduction/CameraCalibrationCore/Source/CameraCalibrationCore/Private/CameraCalibrationSubsystem.cpp
+++ b/Engine/Plugins/VirtualProduction/CameraCalibrationCore/Source/CameraCalibrationCore/Private/CameraCalibrationSubsystem.cpp
@@ -118,9 +118,9 @@ void UCameraCalibrationSubsystem::UnregisterDistortionModel(TSubclassOf<ULensMod
 	LensModelMap.Remove(LensModel->GetDefaultObject<ULensModel>()->GetModelName());
 }
 
-void UCameraCalibrationSubsystem::RegisterOverlayMaterial(const FName& MaterialName, const FName& MaterialPath)
+void UCameraCalibrationSubsystem::RegisterOverlayMaterial(const FName& MaterialName, const FSoftObjectPath& MaterialPath)
 {
-	RegisteredOverlayMaterials.Add(MaterialName, TSoftObjectPtr<UMaterialInterface>(FSoftObjectPath(MaterialPath)));
+	RegisteredOverlayMaterials.Add(MaterialName, TSoftObjectPtr<UMaterialInterface>(MaterialPath));
 }
 
 void UCameraCalibrationSubsystem::UnregisterOverlayMaterial(const FName& MaterialName)
diff --git a/Engine/Plugins/VirtualProduction/CameraCalibrationCore/Source/CameraCalibrationCore/Public/CameraCalibrationSubsystem.h b/Engine/Plugins/VirtualProduction/CameraCalibrationCore/Source/CameraCalibrationCore/Public/CameraCalibrationSubsystem.h
--- a/Engine/Plugins/VirtualProduction/CameraCalibrationCore/Source/CameraCalibrationCore/Public/CameraCalibrationSubsystem.h
+++ b/Engine/Plugins/VirtualProduction/CameraCalibrationCore/Source/CameraCalibrationCore/Public/CameraCalibrationSubsystem.h
@@ -125,7 +125,15 @@ public:
 	void UpdateOverscanFocalLength(UCineCameraComponent* Component, float InFocalLength);
 
 	/** Register a new overlay material name and path that can be queried from camera calibration tools */
-	void RegisterOverlayMaterial(const FName& MaterialName, const FName& MaterialPath);
+	void RegisterOverlayMaterial(const FName& MaterialName, const FSoftObjectPath& MaterialPath);
+
+	UE_DEPRECATED(5.1, "FNames containing full object paths are deprecated. Use FSoftObjectPath instead.")
+	void RegisterOverlayMaterial(const FName& MaterialName, const FName& MaterialPathName)
+	{
+PRAGMA_DISABLE_DEPRECATION_WARNINGS
+		RegisterOverlayMaterial(MaterialName, FSoftObjectPath(MaterialPathName));
+PRAGMA_ENABLE_DEPRECATION_WARNINGS
+	}
 
 	/** Unregister an overlay material */
 	void UnregisterOverlayMaterial(const FName& MaterialName);
diff --git a/Engine/Plugins/VirtualProduction/Takes/Source/TakeRecorder/Private/Widgets/STakeRecorderCockpit.cpp b/Engine/Plugins/VirtualProduction/Takes/Source/TakeRecorder/Private/Widgets/STakeRecorderCockpit.cpp
--- a/Engine/Plugins/VirtualProduction/Takes/Source/TakeRecorder/Private/Widgets/STakeRecorderCockpit.cpp
+++ b/Engine/Plugins/VirtualProduction/Takes/Source/TakeRecorder/Private/Widgets/STakeRecorderCockpit.cpp
@@ -622,7 +622,7 @@ void STakeRecorderCockpit::UpdateTakeError()
 	// If there's only a single one, and it's the one that we're looking at directly, don't show the error
 	if (DuplicateTakes.Num() == 1 && DuplicateTakes[0].IsValid())
 	{
-		ULevelSequence* AlreadyLoaded = FindObject<ULevelSequence>(nullptr, *DuplicateTakes[0].ObjectPath.ToString());
+		ULevelSequence* AlreadyLoaded = FindObject<ULevelSequence>(nullptr, *DuplicateTakes[0].GetObjectPathString());
 		if (AlreadyLoaded && AlreadyLoaded->FindMetaData<UTakeMetaData>() == TakeMetaData)
 		{
 			return;
diff --git a/Engine/Plugins/VirtualProduction/Takes/Source/TakeRecorderSources/Private/TakeRecorderMicrophoneAudioSource.cpp b/Engine/Plugins/VirtualProduction/Takes/Source/TakeRecorderSources/Private/TakeRecorderMicrophoneAudioSource.cpp
--- a/Engine/Plugins/VirtualProduction/Takes/Source/TakeRecorderSources/Private/TakeRecorderMicrophoneAudioSource.cpp
+++ b/Engine/Plugins/VirtualProduction/Takes/Source/TakeRecorderSources/Private/TakeRecorderMicrophoneAudioSource.cpp
@@ -78,7 +78,7 @@ static FString MakeNewAssetName(const FString& BaseAssetPath, const FString& Bas
 	AssetPath += Dot + AssetName;
 
 	FAssetRegistryModule& AssetRegistryModule = FModuleManager::LoadModuleChecked<FAssetRegistryModule>(TEXT("AssetRegistry"));
-	FAssetData AssetData = AssetRegistryModule.Get().GetAssetByObjectPath(*AssetPath);
+	FAssetData AssetData = AssetRegistryModule.Get().GetAssetByObjectPath(FSoftObjectPath(AssetPath));
 
 	// if object with same name exists, try a different name until we don't find one
 	int32 ExtensionIndex = 0;
@@ -86,7 +86,7 @@ static FString MakeNewAssetName(const FString& BaseAssetPath, const FString& Bas
 	{
 		AssetName = FString::Printf(TEXT("%s_%d"), *BaseAssetName, ExtensionIndex);
 		AssetPath = (BaseAssetPath / AssetName) + Dot + AssetName;
-		AssetData = AssetRegistryModule.Get().GetAssetByObjectPath(*AssetPath);
+		AssetData = AssetRegistryModule.Get().GetAssetByObjectPath(FSoftObjectPath(AssetPath));
 
 		ExtensionIndex++;
 	}
diff --git a/Engine/Plugins/VirtualProduction/Takes/Source/TakesCore/Public/TakesUtils.h b/Engine/Plugins/VirtualProduction/Takes/Source/TakesCore/Public/TakesUtils.h
--- a/Engine/Plugins/VirtualProduction/Takes/Source/TakesCore/Public/TakesUtils.h
+++ b/Engine/Plugins/VirtualProduction/Takes/Source/TakesCore/Public/TakesUtils.h
@@ -60,7 +60,7 @@ namespace TakesUtils
 
 		// Generate a unique level sequence name for this take if there are already assets of the same name
 		IAssetRegistry& AssetRegistry = FModuleManager::LoadModuleChecked<FAssetRegistryModule>("AssetRegistry").Get();
-		while (AssetRegistry.GetAssetByObjectPath(*InPackageName).IsValid())
+		while (AssetRegistry.GetAssetByObjectPath(FSoftObjectPath(InPackageName)).IsValid()) // Note: as-is this will only return in-memory assets
 		{
 			int32 TrimCount = InPackageName.Len() - BasePackageLength;
 			if (TrimCount > 0)
@@ -122,7 +122,7 @@ namespace TakesUtils
 		AssetPath += Dot + AssetName;
 
 		FAssetRegistryModule& AssetRegistryModule = FModuleManager::LoadModuleChecked<FAssetRegistryModule>(TEXT("AssetRegistry"));
-		FAssetData AssetData = AssetRegistryModule.Get().GetAssetByObjectPath(*AssetPath);
+		FAssetData AssetData = AssetRegistryModule.Get().GetAssetByObjectPath(FSoftObjectPath(AssetPath));
 
 		// if object with same name exists, try a different name until we don't find one
 		int32 ExtensionIndex = 0;
@@ -130,7 +130,7 @@ namespace TakesUtils
 		{
 			AssetName = FString::Printf(TEXT("%s_%d"), *AssetName, ExtensionIndex);
 			AssetPath = (BaseAssetPath / AssetName) + Dot + AssetName;
-			AssetData = AssetRegistryModule.Get().GetAssetByObjectPath(*AssetPath);
+			AssetData = AssetRegistryModule.Get().GetAssetByObjectPath(FSoftObjectPath(AssetPath));
 
 			ExtensionIndex++;
 		}
