Take Recorder: Change to use GetAssetByPackageName

diff --git a/Engine/Plugins/VirtualProduction/Takes/Source/TakesCore/Public/TakesUtils.h b/Engine/Plugins/VirtualProduction/Takes/Source/TakesCore/Public/TakesUtils.h
--- a/Engine/Plugins/VirtualProduction/Takes/Source/TakesCore/Public/TakesUtils.h
+++ b/Engine/Plugins/VirtualProduction/Takes/Source/TakesCore/Public/TakesUtils.h
@@ -60,7 +60,10 @@ namespace TakesUtils
 
 		// Generate a unique level sequence name for this take if there are already assets of the same name
 		IAssetRegistry& AssetRegistry = FModuleManager::LoadModuleChecked<FAssetRegistryModule>("AssetRegistry").Get();
-		while (AssetRegistry.GetAssetByObjectPath(FSoftObjectPath(InPackageName)).IsValid()) // Note: as-is this will only return in-memory assets
+		TArray<FAssetData> OutAssetData;
+
+		AssetRegistry.GetAssetsByPackageName(*InPackageName, OutAssetData);
+		while (OutAssetData.Num() > 0)
 		{
 			int32 TrimCount = InPackageName.Len() - BasePackageLength;
 			if (TrimCount > 0)
@@ -69,6 +72,9 @@ namespace TakesUtils
 			}
 
 			InPackageName += FString::Printf(TEXT("_%04d"), UniqueIndex++);
+
+			OutAssetData.Empty();
+			AssetRegistry.GetAssetsByPackageName(*InPackageName, OutAssetData);
 		}
 
 		// Create the asset to record into
