Take Recorder: Add SetDefaultParameters for take recorder

diff --git a/Engine/Plugins/VirtualProduction/Takes/Source/TakeRecorder/Private/Recorder/TakeRecorderBlueprintLibrary.cpp b/Engine/Plugins/VirtualProduction/Takes/Source/TakeRecorder/Private/Recorder/TakeRecorderBlueprintLibrary.cpp
--- a/Engine/Plugins/VirtualProduction/Takes/Source/TakeRecorder/Private/Recorder/TakeRecorderBlueprintLibrary.cpp
+++ b/Engine/Plugins/VirtualProduction/Takes/Source/TakeRecorder/Private/Recorder/TakeRecorderBlueprintLibrary.cpp
@@ -82,6 +82,12 @@ FTakeRecorderParameters UTakeRecorderBlueprintLibrary::GetDefaultParameters()
 	return DefaultParams;
 }
 
+void UTakeRecorderBlueprintLibrary::SetDefaultParameters(const FTakeRecorderParameters& InDefaultParameters)
+{
+	GetMutableDefault<UTakeRecorderUserSettings>()->Settings = InDefaultParameters.User;
+	GetMutableDefault<UTakeRecorderProjectSettings>()->Settings = InDefaultParameters.Project;
+}
+
 bool UTakeRecorderBlueprintLibrary::IsRecording()
 {
 	return UTakeRecorder::GetActiveRecorder() != nullptr && UTakeRecorder::GetActiveRecorder()->GetState() == ETakeRecorderState::Started;
diff --git a/Engine/Plugins/VirtualProduction/Takes/Source/TakeRecorder/Public/Recorder/TakeRecorderBlueprintLibrary.h b/Engine/Plugins/VirtualProduction/Takes/Source/TakeRecorder/Public/Recorder/TakeRecorderBlueprintLibrary.h
--- a/Engine/Plugins/VirtualProduction/Takes/Source/TakeRecorder/Public/Recorder/TakeRecorderBlueprintLibrary.h
+++ b/Engine/Plugins/VirtualProduction/Takes/Source/TakeRecorder/Public/Recorder/TakeRecorderBlueprintLibrary.h
@@ -53,6 +53,12 @@ public:
 	static FTakeRecorderParameters GetDefaultParameters();
 
 
+	/**
+	 * Set the default recorder parameters
+	 */
+	UFUNCTION(BlueprintCallable, Category = "Take Recorder")
+	static void SetDefaultParameters(const FTakeRecorderParameters& DefaultParameters);
+
 	/**
 	 * Check whether a recording is currently active
 	 */
