Take Recorder: Don't override save recorded assets to false when recording into a sequence. This should just follow the user's preference

diff --git a/Engine/Plugins/VirtualProduction/Takes/Source/TakeRecorder/Private/Recorder/TakeRecorder.cpp b/Engine/Plugins/VirtualProduction/Takes/Source/TakeRecorder/Private/Recorder/TakeRecorder.cpp
--- a/Engine/Plugins/VirtualProduction/Takes/Source/TakeRecorder/Private/Recorder/TakeRecorder.cpp
+++ b/Engine/Plugins/VirtualProduction/Takes/Source/TakeRecorder/Private/Recorder/TakeRecorder.cpp
@@ -486,7 +486,6 @@ bool UTakeRecorder::Initialize ( ULevelSequence* LevelSequenceBase, UTakeRecorde
 	{
 		Parameters.Project.bStartAtCurrentTimecode = false;
 		Parameters.User.bStopAtPlaybackEnd = true;
-		Parameters.User.bSaveRecordedAssets = false;
 		Parameters.User.bAutoLock = false;
 	}
 
