Allow TakeRecorder.AllowMenuExtensions to be toggled at runtime
#rnx

diff --git a/Engine/Plugins/VirtualProduction/Takes/Source/TakeRecorderSources/Private/TakeRecorderSourcesModule.cpp b/Engine/Plugins/VirtualProduction/Takes/Source/TakeRecorderSources/Private/TakeRecorderSourcesModule.cpp
--- a/Engine/Plugins/VirtualProduction/Takes/Source/TakeRecorderSources/Private/TakeRecorderSourcesModule.cpp
+++ b/Engine/Plugins/VirtualProduction/Takes/Source/TakeRecorderSources/Private/TakeRecorderSourcesModule.cpp
@@ -351,7 +351,7 @@ public:
 	void RegisterMenuExtensions()
 	{
 #if WITH_EDITOR
-		if (GEditor && TakeRecorderSources::AllowMenuExtensions)
+		if (GEditor)
 		{
 			// Register level editor menu extender
 			LevelEditorMenuExtenderDelegate = FLevelEditorModule::FLevelViewportMenuExtender_SelectedActors::CreateRaw(this, &FTakeRecorderSourcesModule::ExtendLevelViewportContextMenu);
@@ -388,6 +388,12 @@ public:
 	{
 		TSharedRef<FExtender> Extender(new FExtender());
 
+#if WITH_EDITOR
+		if (!TakeRecorderSources::AllowMenuExtensions)
+		{
+			return Extender;
+		}
+#endif
 		if (SelectedActors.Num() > 0)
 		{
 			Extender->AddMenuExtension("ActorUETools", EExtensionHook::After, CommandList, FMenuExtensionDelegate::CreateLambda(
