Take Recorder: Change recording to auto keys (not broken). Change AddKeyToSection to AddKeysToSection so that keys can be set all at once (and AutoSetTangents called just once) as opposed to one by one.

Added AddKeys functions for Bool, String, Byte, Integer channels. It's already on Float and Double channels.

diff --git a/Engine/Plugins/VirtualProduction/Takes/Source/TakeTrackRecorders/Private/MovieSceneTrackPropertyRecorder.cpp b/Engine/Plugins/VirtualProduction/Takes/Source/TakeTrackRecorders/Private/MovieSceneTrackPropertyRecorder.cpp
--- a/Engine/Plugins/VirtualProduction/Takes/Source/TakeTrackRecorders/Private/MovieSceneTrackPropertyRecorder.cpp
+++ b/Engine/Plugins/VirtualProduction/Takes/Source/TakeTrackRecorders/Private/MovieSceneTrackPropertyRecorder.cpp
@@ -83,9 +83,20 @@ void FMovieSceneTrackPropertyRecorder<bool>::PostCreate(IMovieSceneTrackRecorder
 }
 
 template <>
-void FMovieSceneTrackPropertyRecorder<bool>::AddKeyToSection(UMovieSceneSection* InSection, const FPropertyKey<bool>& InKey)
+void FMovieSceneTrackPropertyRecorder<bool>::AddKeysToSection(UMovieSceneSection* InSection, const TArray<FPropertyKey<bool>>& InKeys)
 {
-	InSection->GetChannelProxy().GetChannel<FMovieSceneBoolChannel>(0)->GetData().AddKey(InKey.Time, InKey.Value);
+	TArray<FFrameNumber> Times;
+	Times.SetNum(InKeys.Num());
+	TArray<bool> Values;
+	Values.SetNum(InKeys.Num());
+
+	for (int32 Index = 0; Index < InKeys.Num(); ++Index)
+	{
+		Times[Index] = InKeys[Index].Time;
+		Values[Index] = InKeys[Index].Value;
+	}
+
+	InSection->GetChannelProxy().GetChannel<FMovieSceneBoolChannel>(0)->AddKeys(Times, Values);
 }
 
 template <>
@@ -159,6 +170,7 @@ bool FMovieSceneTrackPropertyRecorder<bool>::LoadRecordedFile(const FString& Fil
 						if (InFrames.Num() > 0)
 						{
 							FFrameRate InFrameRate = Header.TickResolution;
+							TArray<FPropertyKey<bool>> NewKeys;
 							for (const FPropertySerializedBoolFrame& SerializedFrame : InFrames)
 							{
 								const FPropertySerializedBool& Frame = SerializedFrame.Frame;
@@ -169,9 +181,10 @@ bool FMovieSceneTrackPropertyRecorder<bool>::LoadRecordedFile(const FString& Fil
 								FPropertyKey<bool> Key;
 								Key.Time = CurrentFrame;
 								Key.Value = Frame.Value;
-								AddKeyToSection(MovieSceneSection.Get(), Key);
+								NewKeys.Add(Key);
 								MovieSceneSection->ExpandToFrame(CurrentFrame);
 							}
+							AddKeysToSection(MovieSceneSection.Get(), NewKeys);
 						}
 						Serializer.Close();
 						InCompletionCallback();
@@ -259,9 +272,20 @@ void FMovieSceneTrackPropertyRecorder<uint8>::PostCreate(IMovieSceneTrackRecorde
 }
 
 template <>
-void FMovieSceneTrackPropertyRecorder<uint8>::AddKeyToSection(UMovieSceneSection* InSection, const FPropertyKey<uint8>& InKey)
+void FMovieSceneTrackPropertyRecorder<uint8>::AddKeysToSection(UMovieSceneSection* InSection, const TArray<FPropertyKey<uint8>>& InKeys)
 {
-	InSection->GetChannelProxy().GetChannel<FMovieSceneByteChannel>(0)->GetData().AddKey(InKey.Time, InKey.Value);
+	TArray<FFrameNumber> Times;
+	Times.SetNum(InKeys.Num());
+	TArray<uint8> Values;
+	Values.SetNum(InKeys.Num());
+
+	for (int32 Index = 0; Index < InKeys.Num(); ++Index)
+	{
+		Times[Index] = InKeys[Index].Time;
+		Values[Index] = InKeys[Index].Value;
+	}
+
+	InSection->GetChannelProxy().GetChannel<FMovieSceneByteChannel>(0)->AddKeys(Times, Values);
 }
 
 template <>
@@ -334,6 +358,7 @@ bool FMovieSceneTrackPropertyRecorder<uint8>::LoadRecordedFile(const FString& Fi
 					if (InFrames.Num() > 0)
 					{
 						FFrameRate InFrameRate = Header.TickResolution;
+						TArray<FPropertyKey<uint8>> NewKeys;
 						for (const FPropertySerializedByteFrame& SerializedFrame : InFrames)
 						{
 							const FPropertySerializedByte& Frame = SerializedFrame.Frame;
@@ -343,9 +368,10 @@ bool FMovieSceneTrackPropertyRecorder<uint8>::LoadRecordedFile(const FString& Fi
 							FPropertyKey<uint8> Key;
 							Key.Time = CurrentFrame;
 							Key.Value = Frame.Value;
-							AddKeyToSection(MovieSceneSection.Get(), Key);
+							NewKeys.Add(Key);
 							MovieSceneSection->ExpandToFrame(CurrentFrame);
 						}
+						AddKeysToSection(MovieSceneSection.Get(), NewKeys);
 					}
 					Serializer.Close();
 					InCompletionCallback();
@@ -412,9 +438,20 @@ void FMovieSceneTrackPropertyRecorder<double>::PostCreate(IMovieSceneTrackRecord
 }
 
 template <>
-void FMovieSceneTrackPropertyRecorder<double>::AddKeyToSection(UMovieSceneSection* InSection, const FPropertyKey<double>& InKey)
+void FMovieSceneTrackPropertyRecorder<double>::AddKeysToSection(UMovieSceneSection* InSection, const TArray<FPropertyKey<double>>& InKeys)
 {
-	InSection->GetChannelProxy().GetChannel<FMovieSceneDoubleChannel>(0)->AddCubicKey(InKey.Time, InKey.Value, RCTM_Break);
+	TArray<FFrameNumber> Times;
+	Times.SetNum(InKeys.Num());
+	TArray<FMovieSceneDoubleValue> Values;
+	Values.SetNum(InKeys.Num());
+
+	for (int32 Index = 0; Index < InKeys.Num(); ++Index)
+	{
+		Times[Index] = InKeys[Index].Time;
+		Values[Index] = FMovieSceneDoubleValue(InKeys[Index].Value);
+	}
+
+	InSection->GetChannelProxy().GetChannel<FMovieSceneDoubleChannel>(0)->AddKeys(Times, Values);
 }
 
 template <>
@@ -490,6 +527,7 @@ bool FMovieSceneTrackPropertyRecorder<double>::LoadRecordedFile(const FString& F
 					if (InFrames.Num() > 0)
 					{
 						FFrameRate InFrameRate = Header.TickResolution;
+						TArray<FPropertyKey<double>> NewKeys;
 						for (const FPropertySerializedDoubleFrame& SerializedFrame : InFrames)
 						{
 							const FPropertySerializedDouble& Frame = SerializedFrame.Frame;
@@ -499,9 +537,10 @@ bool FMovieSceneTrackPropertyRecorder<double>::LoadRecordedFile(const FString& F
 							FPropertyKey<double> Key;
 							Key.Time = CurrentFrame;
 							Key.Value = Frame.Value;
-							AddKeyToSection(MovieSceneSection.Get(), Key);
+							NewKeys.Add(Key);
 							MovieSceneSection->ExpandToFrame(CurrentFrame);
 						}
+						AddKeysToSection(MovieSceneSection.Get(), NewKeys);
 					}
 					Serializer.Close();
 					InCompletionCallback();
@@ -568,9 +607,20 @@ void FMovieSceneTrackPropertyRecorder<float>::PostCreate(IMovieSceneTrackRecorde
 }
 
 template <>
-void FMovieSceneTrackPropertyRecorder<float>::AddKeyToSection(UMovieSceneSection* InSection, const FPropertyKey<float>& InKey)
+void FMovieSceneTrackPropertyRecorder<float>::AddKeysToSection(UMovieSceneSection* InSection, const TArray<FPropertyKey<float>>& InKeys)
 {
-	InSection->GetChannelProxy().GetChannel<FMovieSceneFloatChannel>(0)->AddCubicKey(InKey.Time, InKey.Value, RCTM_Break);
+	TArray<FFrameNumber> Times;
+	Times.SetNum(InKeys.Num());
+	TArray<FMovieSceneFloatValue> Values;
+	Values.SetNum(InKeys.Num());
+	
+	for (int32 Index = 0; Index < InKeys.Num(); ++Index)
+	{
+		Times[Index] = InKeys[Index].Time;
+		Values[Index] = FMovieSceneFloatValue(InKeys[Index].Value);
+	}
+
+	InSection->GetChannelProxy().GetChannel<FMovieSceneFloatChannel>(0)->AddKeys(Times, Values);
 }
 
 template <>
@@ -646,6 +696,7 @@ bool FMovieSceneTrackPropertyRecorder<float>::LoadRecordedFile(const FString& Fi
 					if (InFrames.Num() > 0)
 					{
 						FFrameRate InFrameRate = Header.TickResolution;
+						TArray<FPropertyKey<float>> NewKeys;
 						for (const FPropertySerializedFloatFrame& SerializedFrame : InFrames)
 						{
 							const FPropertySerializedFloat& Frame = SerializedFrame.Frame;
@@ -655,9 +706,10 @@ bool FMovieSceneTrackPropertyRecorder<float>::LoadRecordedFile(const FString& Fi
 							FPropertyKey<float> Key;
 							Key.Time = CurrentFrame;
 							Key.Value = Frame.Value;
-							AddKeyToSection(MovieSceneSection.Get(), Key);
+							NewKeys.Add(Key);
 							MovieSceneSection->ExpandToFrame(CurrentFrame);
 						}
+						AddKeysToSection(MovieSceneSection.Get(), NewKeys);
 					}
 					Serializer.Close();
 					InCompletionCallback();
@@ -728,14 +780,30 @@ void FMovieSceneTrackPropertyRecorder<FColor>::PostCreate(IMovieSceneTrackRecord
 }
 
 template <>
-void FMovieSceneTrackPropertyRecorder<FColor>::AddKeyToSection(UMovieSceneSection* InSection, const FPropertyKey<FColor>& InKey)
+void FMovieSceneTrackPropertyRecorder<FColor>::AddKeysToSection(UMovieSceneSection* InSection, const TArray<FPropertyKey<FColor>>& InKeys)
 {
+	TArray<FFrameNumber> Times;
+	Times.SetNum(InKeys.Num());
+	TArray<FMovieSceneFloatValue> RedValues, GreenValues, BlueValues, AlphaValues;
+	RedValues.SetNum(InKeys.Num());
+	GreenValues.SetNum(InKeys.Num());
+	BlueValues.SetNum(InKeys.Num());
+	AlphaValues.SetNum(InKeys.Num());
+
 	static const float InvColor = 1.0f / 255.0f;
-	TArrayView<FMovieSceneFloatChannel*> FloatChannels = InSection->GetChannelProxy().GetChannels<FMovieSceneFloatChannel>();
-	FloatChannels[0]->AddCubicKey(InKey.Time, InKey.Value.R * InvColor, RCTM_Break);
-	FloatChannels[1]->AddCubicKey(InKey.Time, InKey.Value.G * InvColor, RCTM_Break);
-	FloatChannels[2]->AddCubicKey(InKey.Time, InKey.Value.B * InvColor, RCTM_Break);
-	FloatChannels[3]->AddCubicKey(InKey.Time, InKey.Value.A * InvColor, RCTM_Break);
+	for (int32 Index = 0; Index < InKeys.Num(); ++Index)
+	{
+		Times[Index] = InKeys[Index].Time;
+		RedValues[Index] = FMovieSceneFloatValue(InKeys[Index].Value.R * InvColor);
+		GreenValues[Index] = FMovieSceneFloatValue(InKeys[Index].Value.G * InvColor);
+		BlueValues[Index] = FMovieSceneFloatValue(InKeys[Index].Value.B * InvColor);
+		AlphaValues[Index] = FMovieSceneFloatValue(InKeys[Index].Value.A * InvColor);
+	}
+
+	InSection->GetChannelProxy().GetChannel<FMovieSceneFloatChannel>(0)->AddKeys(Times, RedValues);
+	InSection->GetChannelProxy().GetChannel<FMovieSceneFloatChannel>(1)->AddKeys(Times, GreenValues);
+	InSection->GetChannelProxy().GetChannel<FMovieSceneFloatChannel>(2)->AddKeys(Times, BlueValues);
+	InSection->GetChannelProxy().GetChannel<FMovieSceneFloatChannel>(3)->AddKeys(Times, AlphaValues);
 }
 
 template <>
@@ -850,6 +918,7 @@ bool FMovieSceneTrackPropertyRecorder<FColor>::LoadRecordedFile(const FString& F
 					if (InFrames.Num() > 0)
 					{
 						FFrameRate InFrameRate = Header.TickResolution;
+						TArray<FPropertyKey<FColor>> NewKeys;
 						for (const FPropertySerializedColorFrame& SerializedFrame : InFrames)
 						{
 							const FPropertySerializedColor& Frame = SerializedFrame.Frame;
@@ -859,9 +928,10 @@ bool FMovieSceneTrackPropertyRecorder<FColor>::LoadRecordedFile(const FString& F
 							FPropertyKey<FColor> Key;
 							Key.Time = CurrentFrame;
 							Key.Value = FColor(Frame.Value); 
-							AddKeyToSection(MovieSceneSection.Get(), Key);
+							NewKeys.Add(Key);
 							MovieSceneSection->ExpandToFrame(CurrentFrame);
 						}
+						AddKeysToSection(MovieSceneSection.Get(), NewKeys);
 					}
 					Serializer.Close();
 					InCompletionCallback();
@@ -932,13 +1002,28 @@ void FMovieSceneTrackPropertyRecorder<FLinearColor>::PostCreate(IMovieSceneTrack
 }
 
 template <>
-void FMovieSceneTrackPropertyRecorder<FLinearColor>::AddKeyToSection(UMovieSceneSection* InSection, const FPropertyKey<FLinearColor>& InKey)
+void FMovieSceneTrackPropertyRecorder<FLinearColor>::AddKeysToSection(UMovieSceneSection* InSection, const TArray<FPropertyKey<FLinearColor>>& InKeys)
 {
-	TArrayView<FMovieSceneFloatChannel*> FloatChannels = InSection->GetChannelProxy().GetChannels<FMovieSceneFloatChannel>();
-	FloatChannels[0]->AddCubicKey(InKey.Time, InKey.Value.R, RCTM_Break);
-	FloatChannels[1]->AddCubicKey(InKey.Time, InKey.Value.G, RCTM_Break);
-	FloatChannels[2]->AddCubicKey(InKey.Time, InKey.Value.B, RCTM_Break);
-	FloatChannels[3]->AddCubicKey(InKey.Time, InKey.Value.A, RCTM_Break);
+	TArray<FFrameNumber> Times;
+	Times.SetNum(InKeys.Num());
+	TArray<FMovieSceneFloatValue> RedValues, GreenValues, BlueValues, AlphaValues;
+	RedValues.SetNum(InKeys.Num());
+	GreenValues.SetNum(InKeys.Num());
+	BlueValues.SetNum(InKeys.Num());
+	AlphaValues.SetNum(InKeys.Num());
+
+	for (int32 Index = 0; Index < InKeys.Num(); ++Index)
+	{
+		Times[Index] = InKeys[Index].Time;
+		RedValues[Index] = FMovieSceneFloatValue(InKeys[Index].Value.R);
+		GreenValues[Index] = FMovieSceneFloatValue(InKeys[Index].Value.G);
+		BlueValues[Index] = FMovieSceneFloatValue(InKeys[Index].Value.B);
+		AlphaValues[Index] = FMovieSceneFloatValue(InKeys[Index].Value.A);
+	}
+	InSection->GetChannelProxy().GetChannels<FMovieSceneFloatChannel>()[0]->AddKeys(Times, RedValues);
+	InSection->GetChannelProxy().GetChannels<FMovieSceneFloatChannel>()[1]->AddKeys(Times, GreenValues);
+	InSection->GetChannelProxy().GetChannels<FMovieSceneFloatChannel>()[2]->AddKeys(Times, BlueValues);
+	InSection->GetChannelProxy().GetChannels<FMovieSceneFloatChannel>()[3]->AddKeys(Times, AlphaValues);
 }
 
 template <>
@@ -1056,6 +1141,7 @@ bool FMovieSceneTrackPropertyRecorder<FLinearColor>::LoadRecordedFile(const FStr
 					if (InFrames.Num() > 0)
 					{
 						FFrameRate InFrameRate = Header.TickResolution;
+						TArray<FPropertyKey<FColor>> NewKeys;
 						for (const FPropertySerializedColorFrame& SerializedFrame : InFrames)
 						{
 							const FPropertySerializedColor& Frame = SerializedFrame.Frame;
@@ -1065,9 +1151,10 @@ bool FMovieSceneTrackPropertyRecorder<FLinearColor>::LoadRecordedFile(const FStr
 							FPropertyKey<FColor> Key;
 							Key.Time = CurrentFrame;
 							Key.Value = FColor(Frame.Value); 
-							AddKeyToSection(MovieSceneSection.Get(), Key);
+							NewKeys.Add(Key);
 							MovieSceneSection->ExpandToFrame(CurrentFrame);
 						}
+						AddKeysToSection(MovieSceneSection.Get(), NewKeys);
 					}
 					Serializer.Close();
 					InCompletionCallback();
@@ -1138,12 +1225,25 @@ void FMovieSceneTrackPropertyRecorder<FVector3f>::PostCreate(IMovieSceneTrackRec
 }
 
 template <>
-void FMovieSceneTrackPropertyRecorder<FVector3f>::AddKeyToSection(UMovieSceneSection* InSection, const FPropertyKey<FVector3f>& InKey)
+void FMovieSceneTrackPropertyRecorder<FVector3f>::AddKeysToSection(UMovieSceneSection* InSection, const TArray<FPropertyKey<FVector3f>>& InKeys)
 {
-	TArrayView<FMovieSceneFloatChannel*> FloatChannels = InSection->GetChannelProxy().GetChannels<FMovieSceneFloatChannel>();
-	FloatChannels[0]->AddCubicKey(InKey.Time, InKey.Value.X, RCTM_Break);
-	FloatChannels[1]->AddCubicKey(InKey.Time, InKey.Value.Y, RCTM_Break);
-	FloatChannels[2]->AddCubicKey(InKey.Time, InKey.Value.Z, RCTM_Break);
+	TArray<FFrameNumber> Times;
+	Times.SetNum(InKeys.Num());
+	TArray<FMovieSceneFloatValue> XValues, YValues, ZValues;
+	XValues.SetNum(InKeys.Num());
+	YValues.SetNum(InKeys.Num());
+	ZValues.SetNum(InKeys.Num());
+
+	for (int32 Index = 0; Index < InKeys.Num(); ++Index)
+	{
+		Times[Index] = InKeys[Index].Time;
+		XValues[Index] = FMovieSceneFloatValue(InKeys[Index].Value.X);
+		YValues[Index] = FMovieSceneFloatValue(InKeys[Index].Value.Y);
+		ZValues[Index] = FMovieSceneFloatValue(InKeys[Index].Value.Z);
+	}
+	InSection->GetChannelProxy().GetChannels<FMovieSceneFloatChannel>()[0]->AddKeys(Times, XValues);
+	InSection->GetChannelProxy().GetChannels<FMovieSceneFloatChannel>()[1]->AddKeys(Times, YValues);
+	InSection->GetChannelProxy().GetChannels<FMovieSceneFloatChannel>()[2]->AddKeys(Times, ZValues);
 }
 
 template <>
@@ -1232,6 +1332,7 @@ bool FMovieSceneTrackPropertyRecorder<FVector3f>::LoadRecordedFile(const FString
 					if (InFrames.Num() > 0)
 					{
 						FFrameRate InFrameRate = Header.TickResolution;
+						TArray<FPropertyKey<FVector3f>> NewKeys;
 						for (const FPropertySerializedVector3fFrame& SerializedFrame : InFrames)
 						{
 							const FPropertySerializedVector3f& Frame = SerializedFrame.Frame;
@@ -1241,9 +1342,10 @@ bool FMovieSceneTrackPropertyRecorder<FVector3f>::LoadRecordedFile(const FString
 							FPropertyKey<FVector3f> Key;
 							Key.Time = CurrentFrame;
 							Key.Value = Frame.Value;
-							AddKeyToSection(MovieSceneSection.Get(), Key);
+							NewKeys.Add(Key);
 							MovieSceneSection->ExpandToFrame(CurrentFrame);
 						}
+						AddKeysToSection(MovieSceneSection.Get(), NewKeys);
 					}
 					Serializer.Close();
 					InCompletionCallback();
@@ -1313,12 +1415,25 @@ void FMovieSceneTrackPropertyRecorder<FVector3d>::PostCreate(IMovieSceneTrackRec
 }
 
 template <>
-void FMovieSceneTrackPropertyRecorder<FVector3d>::AddKeyToSection(UMovieSceneSection* InSection, const FPropertyKey<FVector3d>& InKey)
+void FMovieSceneTrackPropertyRecorder<FVector3d>::AddKeysToSection(UMovieSceneSection* InSection, const TArray<FPropertyKey<FVector3d>>& InKeys)
 {
-	TArrayView<FMovieSceneDoubleChannel*> DoubleChannels = InSection->GetChannelProxy().GetChannels<FMovieSceneDoubleChannel>();
-	DoubleChannels[0]->AddCubicKey(InKey.Time, InKey.Value.X, RCTM_Break);
-	DoubleChannels[1]->AddCubicKey(InKey.Time, InKey.Value.Y, RCTM_Break);
-	DoubleChannels[2]->AddCubicKey(InKey.Time, InKey.Value.Z, RCTM_Break);
+	TArray<FFrameNumber> Times;
+	Times.SetNum(InKeys.Num());
+	TArray<FMovieSceneDoubleValue> XValues, YValues, ZValues;
+	XValues.SetNum(InKeys.Num());
+	YValues.SetNum(InKeys.Num());
+	ZValues.SetNum(InKeys.Num());
+
+	for (int32 Index = 0; Index < InKeys.Num(); ++Index)
+	{
+		Times[Index] = InKeys[Index].Time;
+		XValues[Index] = FMovieSceneDoubleValue(InKeys[Index].Value.X);
+		YValues[Index] = FMovieSceneDoubleValue(InKeys[Index].Value.Y);
+		ZValues[Index] = FMovieSceneDoubleValue(InKeys[Index].Value.Z);
+	}
+	InSection->GetChannelProxy().GetChannels<FMovieSceneDoubleChannel>()[0]->AddKeys(Times, XValues);
+	InSection->GetChannelProxy().GetChannels<FMovieSceneDoubleChannel>()[1]->AddKeys(Times, YValues);
+	InSection->GetChannelProxy().GetChannels<FMovieSceneDoubleChannel>()[2]->AddKeys(Times, ZValues);
 }
 
 template <>
@@ -1407,6 +1522,7 @@ bool FMovieSceneTrackPropertyRecorder<FVector3d>::LoadRecordedFile(const FString
 					if (InFrames.Num() > 0)
 					{
 						FFrameRate InFrameRate = Header.TickResolution;
+						TArray<FPropertyKey<FVector3d>> NewKeys;
 						for (const FPropertySerializedVector3dFrame& SerializedFrame : InFrames)
 						{
 							const FPropertySerializedVector3d& Frame = SerializedFrame.Frame;
@@ -1416,9 +1532,10 @@ bool FMovieSceneTrackPropertyRecorder<FVector3d>::LoadRecordedFile(const FString
 							FPropertyKey<FVector3d> Key;
 							Key.Time = CurrentFrame;
 							Key.Value = Frame.Value;
-							AddKeyToSection(MovieSceneSection.Get(), Key);
+							NewKeys.Add(Key);
 							MovieSceneSection->ExpandToFrame(CurrentFrame);
 						}
+						AddKeysToSection(MovieSceneSection.Get(), NewKeys);
 					}
 					Serializer.Close();
 					InCompletionCallback();
@@ -1493,9 +1610,20 @@ void FMovieSceneTrackPropertyRecorder<int32>::PostCreate(IMovieSceneTrackRecorde
 }
 
 template <>
-void FMovieSceneTrackPropertyRecorder<int32>::AddKeyToSection(UMovieSceneSection* InSection, const FPropertyKey<int32>& InKey)
+void FMovieSceneTrackPropertyRecorder<int32>::AddKeysToSection(UMovieSceneSection* InSection, const TArray<FPropertyKey<int32>>& InKeys)
 {
-	InSection->GetChannelProxy().GetChannel<FMovieSceneIntegerChannel>(0)->GetData().AddKey(InKey.Time, InKey.Value);
+	TArray<FFrameNumber> Times;
+	Times.SetNum(InKeys.Num());
+	TArray<int32> Values;
+	Values.SetNum(InKeys.Num());
+
+	for (int32 Index = 0; Index < InKeys.Num(); ++Index)
+	{
+		Times[Index] = InKeys[Index].Time;
+		Values[Index] = InKeys[Index].Value;
+	}
+
+	InSection->GetChannelProxy().GetChannel<FMovieSceneIntegerChannel>(0)->AddKeys(Times, Values);
 }
 
 template <>
@@ -1569,6 +1697,7 @@ bool FMovieSceneTrackPropertyRecorder<int32>::LoadRecordedFile(const FString& Fi
 						if (InFrames.Num() > 0)
 						{
 							FFrameRate InFrameRate = Header.TickResolution;
+							TArray<FPropertyKey<int32>> NewKeys;
 							for (const FPropertySerializedIntegerFrame& SerializedFrame : InFrames)
 							{
 								const FPropertySerializedInteger& Frame = SerializedFrame.Frame;
@@ -1579,9 +1708,10 @@ bool FMovieSceneTrackPropertyRecorder<int32>::LoadRecordedFile(const FString& Fi
 								FPropertyKey<int32> Key;
 								Key.Time = CurrentFrame;
 								Key.Value = Frame.Value;
-								AddKeyToSection(MovieSceneSection.Get(), Key);
+								NewKeys.Add(Key);
 								MovieSceneSection->ExpandToFrame(CurrentFrame);
 							}
+							AddKeysToSection(MovieSceneSection.Get(), NewKeys);
 						}
 						Serializer.Close();
 						InCompletionCallback();
@@ -1659,9 +1789,20 @@ void FMovieSceneTrackPropertyRecorder<FString>::PostCreate(IMovieSceneTrackRecor
 }
 
 template <>
-void FMovieSceneTrackPropertyRecorder<FString>::AddKeyToSection(UMovieSceneSection* InSection, const FPropertyKey<FString>& InKey)
+void FMovieSceneTrackPropertyRecorder<FString>::AddKeysToSection(UMovieSceneSection* InSection, const TArray<FPropertyKey<FString>>& InKeys)
 {
-	InSection->GetChannelProxy().GetChannel<FMovieSceneStringChannel>(0)->GetData().AddKey(InKey.Time, InKey.Value);
+	TArray<FFrameNumber> Times;
+	Times.SetNum(InKeys.Num());
+	TArray<FString> Values;
+	Values.SetNum(InKeys.Num());
+
+	for (int32 Index = 0; Index < InKeys.Num(); ++Index)
+	{
+		Times[Index] = InKeys[Index].Time;
+		Values[Index] = InKeys[Index].Value;
+	}
+
+	InSection->GetChannelProxy().GetChannel<FMovieSceneStringChannel>(0)->AddKeys(Times, Values);
 }
 
 template <>
@@ -1735,6 +1876,7 @@ bool FMovieSceneTrackPropertyRecorder<FString>::LoadRecordedFile(const FString&
 						if (InFrames.Num() > 0)
 						{
 							FFrameRate InFrameRate = Header.TickResolution;
+							TArray<FPropertyKey<FString>> NewKeys;
 							for (const FPropertySerializedStringFrame& SerializedFrame : InFrames)
 							{
 								const FPropertySerializedString& Frame = SerializedFrame.Frame;
@@ -1745,9 +1887,10 @@ bool FMovieSceneTrackPropertyRecorder<FString>::LoadRecordedFile(const FString&
 								FPropertyKey<FString> Key;
 								Key.Time = CurrentFrame;
 								Key.Value = Frame.Value;
-								AddKeyToSection(MovieSceneSection.Get(), Key);
+								NewKeys.Add(Key);
 								MovieSceneSection->ExpandToFrame(CurrentFrame);
 							}
+							AddKeysToSection(MovieSceneSection.Get(), NewKeys);
 						}
 						Serializer.Close();
 						InCompletionCallback();
diff --git a/Engine/Plugins/VirtualProduction/Takes/Source/TakeTrackRecorders/Public/TrackRecorders/MovieSceneTrackPropertyRecorder.h b/Engine/Plugins/VirtualProduction/Takes/Source/TakeTrackRecorders/Public/TrackRecorders/MovieSceneTrackPropertyRecorder.h
--- a/Engine/Plugins/VirtualProduction/Takes/Source/TakeTrackRecorders/Public/TrackRecorders/MovieSceneTrackPropertyRecorder.h
+++ b/Engine/Plugins/VirtualProduction/Takes/Source/TakeTrackRecorders/Public/TrackRecorders/MovieSceneTrackPropertyRecorder.h
@@ -164,10 +164,7 @@ public:
 		// Enable when we finish recording the section
 		MovieSceneSection->SetIsActive(true);
 
-		for (const FPropertyKey<PropertyType>& Key : Keys)
-		{
-			AddKeyToSection(MovieSceneSection.Get(), Key);
-		}
+		AddKeysToSection(MovieSceneSection.Get(), Keys);
 
 		FTrackRecorderSettings TrackRecorderSettings = OwningTakeRecorderSource->GetTrackRecorderSettings();
 
@@ -281,7 +278,7 @@ private:
 	void PostCreate(IMovieSceneTrackRecorderHost* InRecordingHost, UObject* InObjectToRecord, class UMovieScene* InMovieScene, const FGuid& InGuid, bool bOpenSerializer = false);
 
 	/** Helper function, specialized by type, used to add keys to the movie scene section at Finalize() time */
-	void AddKeyToSection(UMovieSceneSection* InSection, const FPropertyKey<PropertyType>& InKey);
+	void AddKeysToSection(UMovieSceneSection* InSection, const TArray<FPropertyKey<PropertyType>>& InKeys);
 
 	/** Helper function, specialized by type, used to reduce keys */
 	void ReduceKeys(UMovieSceneSection* InSection, float ReduceKeysTolerance);
diff --git a/Engine/Source/Runtime/MovieScene/Public/Channels/MovieSceneBoolChannel.h b/Engine/Source/Runtime/MovieScene/Public/Channels/MovieSceneBoolChannel.h
--- a/Engine/Source/Runtime/MovieScene/Public/Channels/MovieSceneBoolChannel.h
+++ b/Engine/Source/Runtime/MovieScene/Public/Channels/MovieSceneBoolChannel.h
@@ -93,6 +93,18 @@ struct MOVIESCENE_API FMovieSceneBoolChannel : public FMovieSceneChannel
 	 */
 	virtual bool Evaluate(FFrameTime InTime, bool& OutValue) const;
 
+	/**
+	 * Add keys with these times to channel. The number of elements in both arrays much match or nothing is added.
+	 * @param InTimes Times to add
+	 * @param InValues Values to add
+	 */
+	FORCEINLINE void AddKeys(const TArray<FFrameNumber>& InTimes, const TArray<bool>& InValues)
+	{
+		check(InTimes.Num() == InValues.Num());
+		Times.Append(InTimes);
+		Values.Append(InValues);
+	}
+
 public:
 
 	// ~ FMovieSceneChannel Interface
diff --git a/Engine/Source/Runtime/MovieScene/Public/Channels/MovieSceneByteChannel.h b/Engine/Source/Runtime/MovieScene/Public/Channels/MovieSceneByteChannel.h
--- a/Engine/Source/Runtime/MovieScene/Public/Channels/MovieSceneByteChannel.h
+++ b/Engine/Source/Runtime/MovieScene/Public/Channels/MovieSceneByteChannel.h
@@ -97,6 +97,18 @@ struct MOVIESCENE_API FMovieSceneByteChannel : public FMovieSceneChannel
 	 */
 	bool Evaluate(FFrameTime InTime, uint8& OutValue) const;
 
+	/**
+	 * Add keys with these times to channel. The number of elements in both arrays much match or nothing is added.
+	 * @param InTimes Times to add
+	 * @param InValues Values to add
+	 */
+	FORCEINLINE void AddKeys(const TArray<FFrameNumber>& InTimes, const TArray<uint8>& InValues)
+	{
+		check(InTimes.Num() == InValues.Num());
+		Times.Append(InTimes);
+		Values.Append(InValues);
+	}
+
 public:
 
 	// ~ FMovieSceneChannel Interface
diff --git a/Engine/Source/Runtime/MovieScene/Public/Channels/MovieSceneIntegerChannel.h b/Engine/Source/Runtime/MovieScene/Public/Channels/MovieSceneIntegerChannel.h
--- a/Engine/Source/Runtime/MovieScene/Public/Channels/MovieSceneIntegerChannel.h
+++ b/Engine/Source/Runtime/MovieScene/Public/Channels/MovieSceneIntegerChannel.h
@@ -111,6 +111,23 @@ struct MOVIESCENE_API FMovieSceneIntegerChannel : public FMovieSceneChannel
 		}
 	}
 
+	/**
+	 * Add keys with these times to channel. The number of elements in both arrays much match or nothing is added.
+	 * @param InTimes Times to add
+	 * @param InValues Values to add
+	 */
+	FORCEINLINE void AddKeys(const TArray<FFrameNumber>& InTimes, const TArray<int32>& InValues)
+	{
+		check(InTimes.Num() == InValues.Num());
+		int32 Index = Times.Num();
+		Times.Append(InTimes);
+		Values.Append(InValues);
+		for (; Index < Times.Num(); ++Index)
+		{
+			KeyHandles.AllocateHandle(Index);
+		}
+	}
+
 public:
 
 	// ~ FMovieSceneChannel Interface
diff --git a/Engine/Source/Runtime/MovieSceneTracks/Public/Channels/MovieSceneStringChannel.h b/Engine/Source/Runtime/MovieSceneTracks/Public/Channels/MovieSceneStringChannel.h
--- a/Engine/Source/Runtime/MovieSceneTracks/Public/Channels/MovieSceneStringChannel.h
+++ b/Engine/Source/Runtime/MovieSceneTracks/Public/Channels/MovieSceneStringChannel.h
@@ -76,6 +76,18 @@ struct MOVIESCENETRACKS_API FMovieSceneStringChannel : public FMovieSceneChannel
 	 */
 	const FString* Evaluate(FFrameTime InTime) const;
 
+	/**
+	 * Add keys with these times to channel. The number of elements in both arrays much match or nothing is added.
+	 * @param InTimes Times to add
+	 * @param InValues Values to add
+	 */
+	FORCEINLINE void AddKeys(const TArray<FFrameNumber>& InTimes, const TArray<FString>& InValues)
+	{
+		check(InTimes.Num() == InValues.Num());
+		Times.Append(InTimes);
+		Values.Append(InValues);
+	}
+
 public:
 
 	// ~ FMovieSceneChannel Interface
