Fixed non-unity/pch errors

diff --git a/Engine/Plugins/Runtime/AR/AppleAR/AppleARKit/Source/AppleARKit/Private/AppleARKitSessionDelegate.cpp b/Engine/Plugins/Runtime/AR/AppleAR/AppleARKit/Source/AppleARKit/Private/AppleARKitSessionDelegate.cpp
--- a/Engine/Plugins/Runtime/AR/AppleAR/AppleARKit/Source/AppleARKit/Private/AppleARKitSessionDelegate.cpp
+++ b/Engine/Plugins/Runtime/AR/AppleAR/AppleARKit/Source/AppleARKit/Private/AppleARKitSessionDelegate.cpp
@@ -5,6 +5,7 @@
 #include "AppleARKitSystem.h"
 #include "AppleARKitModule.h"
 #include "AppleARKitFrame.h"
+#include "RenderingThread.h"
 
 #if SUPPORTS_ARKIT_1_0
 
diff --git a/Engine/Plugins/VirtualProduction/Takes/Source/TakeRecorderSources/Private/TakeRecorderNearbySpawnedActorSource.cpp b/Engine/Plugins/VirtualProduction/Takes/Source/TakeRecorderSources/Private/TakeRecorderNearbySpawnedActorSource.cpp
--- a/Engine/Plugins/VirtualProduction/Takes/Source/TakeRecorderSources/Private/TakeRecorderNearbySpawnedActorSource.cpp
+++ b/Engine/Plugins/VirtualProduction/Takes/Source/TakeRecorderSources/Private/TakeRecorderNearbySpawnedActorSource.cpp
@@ -2,6 +2,7 @@
 
 #include "TakeRecorderNearbySpawnedActorSource.h"
 
+#include "AssetRegistry/AssetData.h"
 #include "GameFramework/Pawn.h"
 #include "TakeRecorderSources.h"
 #include "TakeRecorderSource.h"
diff --git a/Engine/Plugins/VirtualProduction/Takes/Source/TakeTrackRecorders/Private/MovieSceneAnimationTrackRecorder.cpp b/Engine/Plugins/VirtualProduction/Takes/Source/TakeTrackRecorders/Private/MovieSceneAnimationTrackRecorder.cpp
--- a/Engine/Plugins/VirtualProduction/Takes/Source/TakeTrackRecorders/Private/MovieSceneAnimationTrackRecorder.cpp
+++ b/Engine/Plugins/VirtualProduction/Takes/Source/TakeTrackRecorders/Private/MovieSceneAnimationTrackRecorder.cpp
@@ -2,6 +2,7 @@
 
 #include "TrackRecorders/MovieSceneAnimationTrackRecorder.h"
 #include "Animation/Skeleton.h"
+#include "AssetRegistry/AssetData.h"
 #include "Engine/SkeletalMesh.h"
 #include "TrackRecorders/MovieSceneAnimationTrackRecorderSettings.h"
 #include "TakesUtils.h"
diff --git a/Engine/Plugins/VirtualProduction/Takes/Source/TakesCore/Private/TakeRecorderSources.cpp b/Engine/Plugins/VirtualProduction/Takes/Source/TakesCore/Private/TakeRecorderSources.cpp
--- a/Engine/Plugins/VirtualProduction/Takes/Source/TakesCore/Private/TakeRecorderSources.cpp
+++ b/Engine/Plugins/VirtualProduction/Takes/Source/TakesCore/Private/TakeRecorderSources.cpp
@@ -1,6 +1,7 @@
 // Copyright Epic Games, Inc. All Rights Reserved.
 
 #include "TakeRecorderSources.h"
+#include "AssetRegistry/AssetData.h"
 #include "TakeRecorderSource.h"
 #include "TakesCoreLog.h"
 #include "LevelSequence.h"
diff --git a/Engine/Plugins/VirtualProduction/Takes/Source/TakesCore/Private/TakesUtils.cpp b/Engine/Plugins/VirtualProduction/Takes/Source/TakesCore/Private/TakesUtils.cpp
--- a/Engine/Plugins/VirtualProduction/Takes/Source/TakesCore/Private/TakesUtils.cpp
+++ b/Engine/Plugins/VirtualProduction/Takes/Source/TakesCore/Private/TakesUtils.cpp
@@ -1,6 +1,7 @@
 // Copyright Epic Games, Inc. All Rights Reserved.
 
 #include "TakesUtils.h"
+#include "AssetRegistry/AssetData.h"
 #include "TakesCoreLog.h"
 
 #include "Engine/Engine.h"
