Sequencer: Resurrect SetRowIndex to separate recorded multi channel audio to different track rows

diff --git a/Engine/Plugins/VirtualProduction/Takes/Source/TakeRecorderSources/Private/TakeRecorderMicrophoneAudioSource.cpp b/Engine/Plugins/VirtualProduction/Takes/Source/TakeRecorderSources/Private/TakeRecorderMicrophoneAudioSource.cpp
--- a/Engine/Plugins/VirtualProduction/Takes/Source/TakeRecorderSources/Private/TakeRecorderMicrophoneAudioSource.cpp
+++ b/Engine/Plugins/VirtualProduction/Takes/Source/TakeRecorderSources/Private/TakeRecorderMicrophoneAudioSource.cpp
@@ -237,6 +237,7 @@ TArray<UTakeRecorderSource*> UTakeRecorderMicrophoneAudioSource::PostRecording(U
 
 				FFrameNumber RecordStartFrame = MovieScene->GetPlaybackRange().GetLowerBoundValue();
 
+				NewAudioSection->SetRowIndex(RowIndex + 1);
 				NewAudioSection->SetSound(RecordedSoundWave);
 				NewAudioSection->SetRange(TRange<FFrameNumber>(RecordStartFrame, RecordStartFrame + (RecordedSoundWave->GetDuration() * TickResolution).CeilToFrame()));
 				NewAudioSection->TimecodeSource = FTimecode::FromFrameNumber(RecordStartFrame, DisplayRate);
