Take Recorder: Fix cancel take

Delete the sequence asset last (after OnFinished and OnCancelled) so that the ScopedSequencerPanel can clean up and restore back to the pending take sequence before the current sequence being recorded into is deleted

diff --git a/Engine/Plugins/VirtualProduction/Takes/Source/TakeRecorder/Private/Recorder/TakeRecorder.cpp b/Engine/Plugins/VirtualProduction/Takes/Source/TakeRecorder/Private/Recorder/TakeRecorder.cpp
--- a/Engine/Plugins/VirtualProduction/Takes/Source/TakeRecorder/Private/Recorder/TakeRecorder.cpp
+++ b/Engine/Plugins/VirtualProduction/Takes/Source/TakeRecorder/Private/Recorder/TakeRecorder.cpp
@@ -490,12 +490,12 @@ bool UTakeRecorder::Initialize ( ULevelSequence* LevelSequenceBase, UTakeRecorde
 		Parameters.User.bAutoLock = false;
 	}
 
-	// Figure out which world we're recording from
-	DiscoverSourceWorld();
-
 	// Perform any other parameter-configurable initialization. Must have a valid world at this point.
 	InitializeFromParameters();
 
+	// Figure out which world we're recording from
+	DiscoverSourceWorld();
+
 	// Open a recording notification
 
 	if (ShouldShowNotifications())
@@ -1203,27 +1203,6 @@ void UTakeRecorder::StopInternal(const bool bCancelled)
 		}
 	}
 
-	if (Parameters.TakeRecorderMode == ETakeRecorderMode::RecordNewSequence && !bRecordingFinished)
-	{
-		if (GIsEditor)
-		{
-			// Recording was canceled before it started, so delete the asset. Note we can only do this on editor
-			// nodes. On -game nodes, this cannot be performed. This can only happen with Mult-user and -game node
-			// recording.
-			//
-			FAssetRegistryModule::AssetDeleted(SequenceAsset);
-		}
-
-		// Move the asset to the transient package so that new takes with the same number can be created in its place
-		FName DeletedPackageName = MakeUniqueObjectName(nullptr, UPackage::StaticClass(), *(FString(TEXT("/Temp/") + SequenceAsset->GetName() + TEXT("_Cancelled"))));
-		SequenceAsset->GetOutermost()->Rename(*DeletedPackageName.ToString());
-
-		SequenceAsset->ClearFlags(RF_Standalone | RF_Public);
-		SequenceAsset->RemoveFromRoot();
-		SequenceAsset->MarkAsGarbage();
-		SequenceAsset = nullptr;
-	}
-
 	// Perform any other cleanup that has been defined for this recording
 	for (const TFunction<void()>& Cleanup : OnStopCleanup)
 	{
@@ -1249,6 +1228,29 @@ void UTakeRecorder::StopInternal(const bool bCancelled)
 		}
 	}
 
+	// Delete the asset after OnTakeRecorderFinished and OnTakeRecorderCancelled because the ScopedSequencerPanel 
+	// will still have the current sequence before it is returned to the Pending Take.
+	if (Parameters.TakeRecorderMode == ETakeRecorderMode::RecordNewSequence && !bRecordingFinished)
+	{
+		if (GIsEditor)
+		{
+			// Recording was canceled before it started, so delete the asset. Note we can only do this on editor
+			// nodes. On -game nodes, this cannot be performed. This can only happen with Mult-user and -game node
+			// recording.
+			//
+			FAssetRegistryModule::AssetDeleted(SequenceAsset);
+		}
+
+		// Move the asset to the transient package so that new takes with the same number can be created in its place
+		FName DeletedPackageName = MakeUniqueObjectName(nullptr, UPackage::StaticClass(), *(FString(TEXT("/Temp/") + SequenceAsset->GetName() + TEXT("_Cancelled"))));
+		SequenceAsset->GetOutermost()->Rename(*DeletedPackageName.ToString());
+
+		SequenceAsset->ClearFlags(RF_Standalone | RF_Public);
+		SequenceAsset->RemoveFromRoot();
+		SequenceAsset->MarkAsGarbage();
+		SequenceAsset = nullptr;
+	}
+
 	if (SequenceAsset)
 	{
 		UPackage* const Package = SequenceAsset->GetOutermost();
