Take Recorder: Added AudioAssetName to allow the user to format the name of the audio asset file that is recorded. Defaults to "Audio_{slate}_{take}"

diff --git a/Engine/Plugins/VirtualProduction/Takes/Source/TakeRecorderSources/Private/TakeRecorderMicrophoneAudioSource.cpp b/Engine/Plugins/VirtualProduction/Takes/Source/TakeRecorderSources/Private/TakeRecorderMicrophoneAudioSource.cpp
--- a/Engine/Plugins/VirtualProduction/Takes/Source/TakeRecorderSources/Private/TakeRecorderMicrophoneAudioSource.cpp
+++ b/Engine/Plugins/VirtualProduction/Takes/Source/TakeRecorderSources/Private/TakeRecorderMicrophoneAudioSource.cpp
@@ -29,6 +29,7 @@
 UTakeRecorderMicrophoneAudioSourceSettings::UTakeRecorderMicrophoneAudioSourceSettings(const FObjectInitializer& ObjInit)
 	: Super(ObjInit)
 	, AudioTrackName(NSLOCTEXT("UTakeRecorderMicrophoneAudioSource", "DefaultAudioTrackName", "Recorded Audio"))
+	, AudioAssetName(TEXT("Audio_{slate}_{take}"))
 	, AudioSubDirectory(TEXT("Audio"))
 {
 	TrackTint = FColor(75, 67, 148);
@@ -48,7 +49,7 @@ FString UTakeRecorderMicrophoneAudioSourceSettings::GetSubsceneTrackName(ULevelS
 {
 	if (UTakeMetaData* TakeMetaData = InSequence->FindMetaData<UTakeMetaData>())
 	{
-		return FString::Printf(TEXT("Audio_%s"), *TakeMetaData->GenerateAssetPath("{slate}"));
+		return FString::Printf(TEXT("%s_%s"), *AudioTrackName.ToString(), *TakeMetaData->GenerateAssetPath("{slate}"));
 	}
 	return TEXT("MicrophoneAudio");
 }
@@ -57,7 +58,7 @@ FString UTakeRecorderMicrophoneAudioSourceSettings::GetSubsceneAssetName(ULevelS
 {
 	if (UTakeMetaData* TakeMetaData = InSequence->FindMetaData<UTakeMetaData>())
 	{
-		return FString::Printf(TEXT("Audio_%s"), *TakeMetaData->GenerateAssetPath("{slate}_{take}"));
+		return TakeMetaData->GenerateAssetPath(AudioAssetName);
 	}
 	return TEXT("MicrophoneAudio");
 }
diff --git a/Engine/Plugins/VirtualProduction/Takes/Source/TakeRecorderSources/Private/TakeRecorderMicrophoneAudioSource.h b/Engine/Plugins/VirtualProduction/Takes/Source/TakeRecorderSources/Private/TakeRecorderMicrophoneAudioSource.h
--- a/Engine/Plugins/VirtualProduction/Takes/Source/TakeRecorderSources/Private/TakeRecorderMicrophoneAudioSource.h
+++ b/Engine/Plugins/VirtualProduction/Takes/Source/TakeRecorderSources/Private/TakeRecorderMicrophoneAudioSource.h
@@ -26,10 +26,14 @@ public:
 	virtual FString GetSubsceneAssetName(ULevelSequence* InSequence) const override;
 	// ~UTakeRecorderSource Interface
 
-	/** Name of the recorded audio track name */
+	/** Name of the recorded audio track */
 	UPROPERTY(config, EditAnywhere, BlueprintReadWrite, Category = "Source")
 	FText AudioTrackName;
 
+	/** The name of the audio asset. */
+	UPROPERTY(config, EditAnywhere, BlueprintReadWrite, Category = "Source")
+	FString AudioAssetName;
+
 	/** The name of the subdirectory audio will be placed in. Leave this empty to place into the same directory as the sequence base path */
 	UPROPERTY(config, EditAnywhere, BlueprintReadWrite, Category = "Source")
 	FString AudioSubDirectory;
