Fix non-deterministic Blueprint compilation causing CIS errors.
Blueprints dependended on C++ types loaded in PostEngineInit.

diff --git a/Engine/Plugins/VirtualProduction/Takes/Takes.uplugin b/Engine/Plugins/VirtualProduction/Takes/Takes.uplugin
--- a/Engine/Plugins/VirtualProduction/Takes/Takes.uplugin
+++ b/Engine/Plugins/VirtualProduction/Takes/Takes.uplugin
@@ -25,22 +25,22 @@
 		{
 			"Name": "TakesCore",
 			"Type": "UncookedOnly",
-			"LoadingPhase": "PostEngineInit"
+			"LoadingPhase": "PostDefault"
 		},
 		{
 			"Name": "TakeTrackRecorders",
 			"Type": "UncookedOnly",
-			"LoadingPhase": "PostEngineInit"
+			"LoadingPhase": "PostDefault"
 		},
 		{
 			"Name": "TakeRecorderSources",
 			"Type": "UncookedOnly",
-			"LoadingPhase": "PostEngineInit"
+			"LoadingPhase": "PostDefault"
 		},
 		{
 			"Name": "TakeRecorder",
 			"Type": "UncookedOnly",
-			"LoadingPhase": "PostEngineInit"
+			"LoadingPhase": "PostDefault"
 		},
 		{
 			"Name": "TakeMovieScene",
