Take Recorder: Check read-onlyness before saving

diff --git a/Engine/Plugins/VirtualProduction/Takes/Source/TakesCore/Private/TakesUtils.cpp b/Engine/Plugins/VirtualProduction/Takes/Source/TakesCore/Private/TakesUtils.cpp
--- a/Engine/Plugins/VirtualProduction/Takes/Source/TakesCore/Private/TakesUtils.cpp
+++ b/Engine/Plugins/VirtualProduction/Takes/Source/TakesCore/Private/TakesUtils.cpp
@@ -11,6 +11,7 @@
 #include "Math/Range.h"
 #include "MovieSceneSection.h"
 #include "MovieSceneTimeHelpers.h"
+#include "HAL/FileManager.h"
 #include "Tracks/MovieSceneCameraCutTrack.h"
 #include "Sections/MovieSceneCameraCutSection.h"
 #include "UObject/SavePackage.h"
@@ -95,6 +96,12 @@ void SaveAsset(UObject* InObject)
 	FString const PackageName = Package->GetName();
 	FString const PackageFileName = FPackageName::LongPackageNameToFilename(PackageName, FPackageName::GetAssetPackageExtension());
 
+	if (IFileManager::Get().IsReadOnly(*PackageFileName))
+	{
+		UE_LOG(LogTakesCore, Error, TEXT("Could not save read only file: %s"), *PackageFileName);
+		return;
+	}
+
 	double StartTime = FPlatformTime::Seconds();
 
 	UMetaData *MetaData = Package->GetMetaData();
