AssetDefinition - LevelSequence has been upgraded and no longer loads on right click.
AssetDefinition - Upgraded ReverbEffect.

diff --git a/Engine/Plugins/MovieScene/LevelSequenceEditor/Source/LevelSequenceEditor/LevelSequenceEditor.Build.cs b/Engine/Plugins/MovieScene/LevelSequenceEditor/Source/LevelSequenceEditor/LevelSequenceEditor.Build.cs
--- a/Engine/Plugins/MovieScene/LevelSequenceEditor/Source/LevelSequenceEditor/LevelSequenceEditor.Build.cs
+++ b/Engine/Plugins/MovieScene/LevelSequenceEditor/Source/LevelSequenceEditor/LevelSequenceEditor.Build.cs
@@ -50,6 +50,7 @@ public class LevelSequenceEditor : ModuleRules
 				"TimeManagement",
 				"ToolMenus",
 				"ToolWidgets",
+				"AssetDefinition",
 			}
 		);
 
diff --git a/Engine/Plugins/MovieScene/LevelSequenceEditor/Source/LevelSequenceEditor/Private/AssetTools/AssetDefinition_LevelSequence.cpp b/Engine/Plugins/MovieScene/LevelSequenceEditor/Source/LevelSequenceEditor/Private/AssetTools/AssetDefinition_LevelSequence.cpp
new file mode 100644
--- /dev/null
+++ b/Engine/Plugins/MovieScene/LevelSequenceEditor/Source/LevelSequenceEditor/Private/AssetTools/AssetDefinition_LevelSequence.cpp
@@ -0,0 +1,45 @@
+// Copyright Epic Games, Inc. All Rights Reserved.
+
+#include "AssetDefinition_LevelSequence.h"
+#include "Engine/Engine.h"
+#include "LevelSequence.h"
+#include "LevelSequenceEditorToolkit.h"
+#include "LevelSequenceEditorStyle.h"
+
+#define LOCTEXT_NAMESPACE "AssetTypeActions"
+
+FAssetOpenSupport UAssetDefinition_LevelSequence::GetAssetOpenSupport(const FAssetOpenSupportArgs& OpenSupportArgs) const
+{
+	return FAssetOpenSupport(OpenSupportArgs.OpenMethod,OpenSupportArgs.OpenMethod == EAssetOpenMethod::Edit, EToolkitMode::WorldCentric); 
+}
+
+EAssetCommandResult UAssetDefinition_LevelSequence::OpenAssets(const FAssetOpenArgs& OpenArgs) const
+{
+	UWorld* WorldContext = nullptr;
+	for (const FWorldContext& Context : GEngine->GetWorldContexts())
+	{
+		if (Context.WorldType == EWorldType::Editor)
+		{
+			WorldContext = Context.World();
+			break;
+		}
+	}
+
+	if (!ensure(WorldContext))
+	{
+		return EAssetCommandResult::Handled;
+	}
+
+	for (ULevelSequence* LevelSequence : OpenArgs.LoadObjects<ULevelSequence>())
+	{
+		// Legacy upgrade
+		LevelSequence->ConvertPersistentBindingsToDefault(WorldContext);
+
+		TSharedRef<FLevelSequenceEditorToolkit> Toolkit = MakeShareable(new FLevelSequenceEditorToolkit(FLevelSequenceEditorStyle::Get()));
+		Toolkit->Initialize(OpenArgs.GetToolkitMode(), OpenArgs.ToolkitHost, LevelSequence);
+	}
+
+	return EAssetCommandResult::Handled;
+}
+
+#undef LOCTEXT_NAMESPACE
diff --git a/Engine/Plugins/MovieScene/LevelSequenceEditor/Source/LevelSequenceEditor/Private/AssetTools/AssetDefinition_LevelSequence.h b/Engine/Plugins/MovieScene/LevelSequenceEditor/Source/LevelSequenceEditor/Private/AssetTools/AssetDefinition_LevelSequence.h
new file mode 100644
--- /dev/null
+++ b/Engine/Plugins/MovieScene/LevelSequenceEditor/Source/LevelSequenceEditor/Private/AssetTools/AssetDefinition_LevelSequence.h
@@ -0,0 +1,29 @@
+// Copyright Epic Games, Inc. All Rights Reserved.
+
+#pragma once
+
+#include "CoreMinimal.h"
+#include "LevelSequence.h"
+#include "AssetDefinitionDefault.h"
+
+#include "AssetDefinition_LevelSequence.generated.h"
+
+UCLASS()
+class UAssetDefinition_LevelSequence : public UAssetDefinitionDefault
+{
+	GENERATED_BODY()
+
+public:
+	// UAssetDefinition Begin
+	virtual FText GetAssetDisplayName() const override { return NSLOCTEXT("AssetTypeActions", "AssetTypeActions_LevelSequence", "Level Sequence"); }
+	virtual FLinearColor GetAssetColor() const override { return FLinearColor(FColor(200, 80, 80)); }
+	virtual TSoftClassPtr<UObject> GetAssetClass() const override { return ULevelSequence::StaticClass(); }
+	virtual TConstArrayView<FAssetCategoryPath> GetAssetCategories() const override
+	{
+		static const auto Categories = { EAssetCategoryPaths::Cinematics };
+		return Categories;
+	}
+	virtual FAssetOpenSupport GetAssetOpenSupport(const FAssetOpenSupportArgs& OpenSupportArgs) const override;
+	virtual EAssetCommandResult OpenAssets(const FAssetOpenArgs& OpenArgs) const override;
+	// UAssetDefinition End
+};
diff --git a/Engine/Plugins/MovieScene/LevelSequenceEditor/Source/LevelSequenceEditor/Private/AssetTools/LevelSequenceActions.cpp b/Engine/Plugins/MovieScene/LevelSequenceEditor/Source/LevelSequenceEditor/Private/AssetTools/LevelSequenceActions.cpp
deleted file mode 100644
diff --git a/Engine/Plugins/MovieScene/LevelSequenceEditor/Source/LevelSequenceEditor/Private/AssetTools/LevelSequenceActions.h b/Engine/Plugins/MovieScene/LevelSequenceEditor/Source/LevelSequenceEditor/Private/AssetTools/LevelSequenceActions.h
deleted file mode 100644
diff --git a/Engine/Plugins/MovieScene/LevelSequenceEditor/Source/LevelSequenceEditor/Private/LevelSequenceEditorModule.cpp b/Engine/Plugins/MovieScene/LevelSequenceEditor/Source/LevelSequenceEditor/Private/LevelSequenceEditorModule.cpp
--- a/Engine/Plugins/MovieScene/LevelSequenceEditor/Source/LevelSequenceEditor/Private/LevelSequenceEditorModule.cpp
+++ b/Engine/Plugins/MovieScene/LevelSequenceEditor/Source/LevelSequenceEditor/Private/LevelSequenceEditorModule.cpp
@@ -12,7 +12,6 @@
 #include "Framework/MultiBox/MultiBoxBuilder.h"
 
 #include "IAssetTypeActions.h"
-#include "AssetTools/LevelSequenceActions.h"
 #include "LevelSequenceEditorCommands.h"
 #include "Misc/LevelSequenceEditorSettings.h"
 #include "Misc/LevelSequenceEditorHelpers.h"
@@ -67,7 +66,6 @@ public:
 
 		RegisterEditorObjectBindings();
 		RegisterEditorActorSpawner();
-		RegisterAssetTools();
 		RegisterMenuExtensions();
 		RegisterLevelEditorExtensions();
 		RegisterPlacementModeExtensions();
@@ -95,7 +93,6 @@ public:
 	{
 		UnregisterEditorObjectBindings();
 		UnregisterEditorActorSpawner();
-		UnregisterAssetTools();
 		UnregisterMenuExtensions();
 		UnregisterLevelEditorExtensions();
 		UnregisterPlacementModeExtensions();
@@ -119,15 +116,6 @@ protected:
 		EditorActorSpawnerDelegateHandle = LevelSequenceModule.RegisterObjectSpawner(FOnCreateMovieSceneObjectSpawner::CreateStatic(&FLevelSequenceEditorActorSpawner::CreateObjectSpawner));
 	}
 
-	/** Registers asset tool actions. */
-	void RegisterAssetTools()
-	{
-		IAssetTools& AssetTools = FModuleManager::LoadModuleChecked<FAssetToolsModule>("AssetTools").Get();
-		EAssetTypeCategories::Type CinematicAssetCategoryBit = AssetTools.RegisterAdvancedAssetCategory(FName(TEXT("Cinematics")), LOCTEXT("CinematicsAssetCategory", "Cinematics"));
-		LevelSequenceTypeActions = MakeShared<FLevelSequenceActions>(FLevelSequenceEditorStyle::Get(), CinematicAssetCategoryBit);
-		AssetTools.RegisterAssetTypeActions(LevelSequenceTypeActions.ToSharedRef());
-	}
-
 	/** Registers level editor extensions. */
 	void RegisterLevelEditorExtensions()
 	{
@@ -232,18 +220,6 @@ protected:
 		}
 	}
 
-	/** Unregisters asset tool actions. */
-	void UnregisterAssetTools()
-	{
-		FAssetToolsModule* AssetToolsModule = FModuleManager::GetModulePtr<FAssetToolsModule>("AssetTools");
-
-		if (AssetToolsModule != nullptr)
-		{
-			IAssetTools& AssetTools = AssetToolsModule->Get();
-			AssetTools.UnregisterAssetTypeActions(LevelSequenceTypeActions.ToSharedRef());
-		}
-	}
-
 	/** Unregisters level editor extensions. */
 	void UnregisterLevelEditorExtensions()
 	{
@@ -392,9 +368,6 @@ protected:
 
 private:
 
-	/** The type actions for interacting with level sequences. */
-	TSharedPtr<FLevelSequenceActions> LevelSequenceTypeActions;
-
 	/** Extender for the cinematics menu */
 	TSharedPtr<FExtender> CinematicsMenuExtender;
 
diff --git a/Engine/Plugins/MovieScene/MovieRenderPipeline/Source/MovieRenderPipelineEditor/Private/MoviePipelineEditorBlueprintLibrary.cpp b/Engine/Plugins/MovieScene/MovieRenderPipeline/Source/MovieRenderPipelineEditor/Private/MoviePipelineEditorBlueprintLibrary.cpp
--- a/Engine/Plugins/MovieScene/MovieRenderPipeline/Source/MovieRenderPipelineEditor/Private/MoviePipelineEditorBlueprintLibrary.cpp
+++ b/Engine/Plugins/MovieScene/MovieRenderPipeline/Source/MovieRenderPipelineEditor/Private/MoviePipelineEditorBlueprintLibrary.cpp
@@ -158,7 +158,7 @@ UMoviePipelineExecutorJob* UMoviePipelineEditorBlueprintLibrary::CreateJobFromSe
 
 	NewJob->Modify();
 
-	TArray<FString> AssociatedMaps = FSequencerUtilities::GetAssociatedMapPackages(InSequence);
+	TArray<FString> AssociatedMaps = FSequencerUtilities::GetAssociatedLevelSequenceMapPackages(InSequence);
 	FSoftObjectPath CurrentWorld;
 
 	UWorld* EditorWorld = GEditor ? GEditor->GetEditorWorldContext().World() : nullptr;
diff --git a/Engine/Plugins/VirtualProduction/Takes/Source/TakeRecorder/Private/TakeRecorderModule.cpp b/Engine/Plugins/VirtualProduction/Takes/Source/TakeRecorder/Private/TakeRecorderModule.cpp
--- a/Engine/Plugins/VirtualProduction/Takes/Source/TakeRecorder/Private/TakeRecorderModule.cpp
+++ b/Engine/Plugins/VirtualProduction/Takes/Source/TakeRecorder/Private/TakeRecorderModule.cpp
@@ -194,15 +194,17 @@ void FTakeRecorderModule::RegisterMenus()
 
 		if (Context->SelectedAssets.Num() == 1 && Context->SelectedAssets[0].IsInstanceOf(ULevelSequence::StaticClass()))
 		{
-			if (ULevelSequence* LevelSequence = Cast<ULevelSequence>(Context->SelectedAssets[0].GetAsset()))
-			{
-				InSection.AddMenuEntry(
-					"OpenInTakeRecorder_Label",
-					LOCTEXT("OpenInTakeRecorder_Label", "Open in Take Recorder"),
-					LOCTEXT("OpenInTakeRecorder_Tooltip", "Opens this level sequence asset in Take Recorder by copying its contents into the pending take"),
-					FSlateIcon(FTakeRecorderStyle::StyleName, "TakeRecorder.TabIcon"),
-					FExecuteAction::CreateLambda(
-						[LevelSequence]
+			const FAssetData LevelSequenceAsset = Context->SelectedAssets[0];
+
+			InSection.AddMenuEntry(
+				"OpenInTakeRecorder_Label",
+				LOCTEXT("OpenInTakeRecorder_Label", "Open in Take Recorder"),
+				LOCTEXT("OpenInTakeRecorder_Tooltip", "Opens this level sequence asset in Take Recorder by copying its contents into the pending take"),
+				FSlateIcon(FTakeRecorderStyle::StyleName, "TakeRecorder.TabIcon"),
+				FExecuteAction::CreateLambda(
+					[LevelSequenceAsset]
+					{
+						if (ULevelSequence* LevelSequence = Cast<ULevelSequence>(LevelSequenceAsset.GetAsset()))
 						{
 							FLevelEditorModule& LevelEditorModule = FModuleManager::LoadModuleChecked<FLevelEditorModule>("LevelEditor");
 							TSharedPtr<SDockTab> DockTab = LevelEditorModule.GetLevelEditorTabManager()->TryInvokeTab(ITakeRecorderModule::TakeRecorderTabName);
@@ -222,15 +224,18 @@ void FTakeRecorderModule::RegisterMenus()
 								}
 							}
 						}
-					)
-				);
-				InSection.AddMenuEntry(
-					"RecordIntoTakeRecorder_Label",
-					LOCTEXT("RecordWithTakeRecorder_Label", "Record with Take Recorder"),
-					LOCTEXT("RecordWithTakeRecorder_Tooltip", "Opens this level sequence asset for recording into with Take Recorder"),
-					FSlateIcon(FTakeRecorderStyle::StyleName, "TakeRecorder.TabIcon"),
-					FExecuteAction::CreateLambda(
-						[LevelSequence]
+					}
+				)
+			);
+			InSection.AddMenuEntry(
+				"RecordIntoTakeRecorder_Label",
+				LOCTEXT("RecordWithTakeRecorder_Label", "Record with Take Recorder"),
+				LOCTEXT("RecordWithTakeRecorder_Tooltip", "Opens this level sequence asset for recording into with Take Recorder"),
+				FSlateIcon(FTakeRecorderStyle::StyleName, "TakeRecorder.TabIcon"),
+				FExecuteAction::CreateLambda(
+					[LevelSequenceAsset]
+					{
+						if (ULevelSequence* LevelSequence = Cast<ULevelSequence>(LevelSequenceAsset.GetAsset()))
 						{
 							FLevelEditorModule& LevelEditorModule = FModuleManager::LoadModuleChecked<FLevelEditorModule>("LevelEditor");
 							TSharedPtr<SDockTab> DockTab = LevelEditorModule.GetLevelEditorTabManager()->TryInvokeTab(ITakeRecorderModule::TakeRecorderTabName);
@@ -241,9 +246,9 @@ void FTakeRecorderModule::RegisterMenus()
 								TabContent->SetupForRecordingInto(LevelSequence);
 							}
 						}
-					)
-				);
-			}
+					}
+				)
+			);
 		}
 	}));
 #endif // WITH_EDITOR
diff --git a/Engine/Source/Editor/AssetDefinition/Private/AssetDefinition.cpp b/Engine/Source/Editor/AssetDefinition/Private/AssetDefinition.cpp
--- a/Engine/Source/Editor/AssetDefinition/Private/AssetDefinition.cpp
+++ b/Engine/Source/Editor/AssetDefinition/Private/AssetDefinition.cpp
@@ -18,6 +18,7 @@ FAssetCategoryPath EAssetCategoryPaths::Blueprint(LOCTEXT("Blueprint", "Blueprin
 FAssetCategoryPath EAssetCategoryPaths::Texture(LOCTEXT("Texture", "Texture"));
 FAssetCategoryPath EAssetCategoryPaths::Foliage(LOCTEXT("Foliage", "Foliage"));
 FAssetCategoryPath EAssetCategoryPaths::Input(LOCTEXT("Input", "Input"));
+FAssetCategoryPath EAssetCategoryPaths::Cinematics(LOCTEXT("Cinematics", "Cinematics"));
 
 FAssetCategoryPath::FAssetCategoryPath(const FText& InCategory)
 {
diff --git a/Engine/Source/Editor/AssetDefinition/Public/AssetDefinition.h b/Engine/Source/Editor/AssetDefinition/Public/AssetDefinition.h
--- a/Engine/Source/Editor/AssetDefinition/Public/AssetDefinition.h
+++ b/Engine/Source/Editor/AssetDefinition/Public/AssetDefinition.h
@@ -234,16 +234,25 @@ struct ASSETDEFINITION_API EAssetCategoryPaths
 	static FAssetCategoryPath Physics;
 	static FAssetCategoryPath Texture;
 	static FAssetCategoryPath UI;
+	static FAssetCategoryPath Cinematics;
 };
 
 struct FAssetOpenSupport
 {
+public:
 	FAssetOpenSupport(EAssetOpenMethod InOpenMethod, bool bInSupported)
 		: OpenMethod(InOpenMethod)
 		, IsSupported(bInSupported)
 	{
 	}
 	
+	FAssetOpenSupport(EAssetOpenMethod InOpenMethod, bool bInSupported, EToolkitMode::Type InRequiredToolkitMode)
+		: OpenMethod(InOpenMethod)
+		, IsSupported(bInSupported)
+		, RequiredToolkitMode(InRequiredToolkitMode)
+	{
+	}
+	
 	EAssetOpenMethod OpenMethod;
 	bool IsSupported;
 	TOptional<EToolkitMode::Type> RequiredToolkitMode;
diff --git a/Engine/Source/Editor/AudioEditor/AudioEditor.Build.cs b/Engine/Source/Editor/AudioEditor/AudioEditor.Build.cs
--- a/Engine/Source/Editor/AudioEditor/AudioEditor.Build.cs
+++ b/Engine/Source/Editor/AudioEditor/AudioEditor.Build.cs
@@ -18,7 +18,8 @@ public class AudioEditor : ModuleRules
 				"DeveloperSettings",
 				"UMGEditor",
 				"AudioExtensions",
-				"AudioLinkEngine"
+				"AudioLinkEngine",
+				"AssetDefinition"
 			}
 		);
 
diff --git a/Engine/Source/Editor/AudioEditor/Private/AssetTypeActions/AssetDefinition_ReverbEffect.h b/Engine/Source/Editor/AudioEditor/Private/AssetTypeActions/AssetDefinition_ReverbEffect.h
new file mode 100644
--- /dev/null
+++ b/Engine/Source/Editor/AudioEditor/Private/AssetTypeActions/AssetDefinition_ReverbEffect.h
@@ -0,0 +1,27 @@
+// Copyright Epic Games, Inc. All Rights Reserved.
+
+#pragma once
+
+#include "CoreMinimal.h"
+#include "Sound/ReverbEffect.h"
+#include "AssetDefinitionDefault.h"
+
+#include "AssetDefinition_ReverbEffect.generated.h"
+
+UCLASS()
+class UAssetDefinition_ReverbEffect : public UAssetDefinitionDefault
+{
+	GENERATED_BODY()
+
+public:
+	// UAssetDefinition Begin
+	virtual FText GetAssetDisplayName() const override { return NSLOCTEXT("AssetTypeActions", "AssetTypeActions_ReverbEffect", "Reverb Effect"); }
+	virtual FLinearColor GetAssetColor() const override { return FLinearColor(FColor(175, 0, 0)); }
+	virtual TSoftClassPtr<UObject> GetAssetClass() const override { return UReverbEffect::StaticClass(); }
+	virtual TConstArrayView<FAssetCategoryPath> GetAssetCategories() const override
+	{
+		static const auto Categories = { FAssetCategoryPath(EAssetCategoryPaths::Audio, NSLOCTEXT("AssetTypeActions", "AssetEffectSubMenu", "Effects")) };
+		return Categories;
+	}
+	// UAssetDefinition End
+};
diff --git a/Engine/Source/Editor/AudioEditor/Private/AssetTypeActions/AssetTypeActions_ReverbEffect.cpp b/Engine/Source/Editor/AudioEditor/Private/AssetTypeActions/AssetTypeActions_ReverbEffect.cpp
deleted file mode 100644
diff --git a/Engine/Source/Editor/AudioEditor/Private/AssetTypeActions/AssetTypeActions_ReverbEffect.h b/Engine/Source/Editor/AudioEditor/Private/AssetTypeActions/AssetTypeActions_ReverbEffect.h
deleted file mode 100644
diff --git a/Engine/Source/Editor/AudioEditor/Private/AudioEditorModule.cpp b/Engine/Source/Editor/AudioEditor/Private/AudioEditorModule.cpp
--- a/Engine/Source/Editor/AudioEditor/Private/AudioEditorModule.cpp
+++ b/Engine/Source/Editor/AudioEditor/Private/AudioEditorModule.cpp
@@ -6,7 +6,6 @@
 #include "AssetTypeActions/AssetTypeActions_AudioBus.h"
 #include "AssetTypeActions/AssetTypeActions_DialogueVoice.h"
 #include "AssetTypeActions/AssetTypeActions_DialogueWave.h"
-#include "AssetTypeActions/AssetTypeActions_ReverbEffect.h"
 #include "AssetTypeActions/AssetTypeActions_SoundAttenuation.h"
 #include "AssetTypeActions/AssetTypeActions_SoundBase.h"
 #include "AssetTypeActions/AssetTypeActions_SoundClass.h"
@@ -195,7 +194,6 @@ public:
 		AssetTools.RegisterAssetTypeActions(MakeShared<FAssetTypeActions_SoundCue>());
 		AssetTools.RegisterAssetTypeActions(MakeShared<FAssetTypeActions_SoundMix>());
 		AssetTools.RegisterAssetTypeActions(MakeShared<FAssetTypeActions_SoundWave>());
-		AssetTools.RegisterAssetTypeActions(MakeShared<FAssetTypeActions_ReverbEffect>());
 	}
 
 	virtual void RegisterAudioMixerAssetActions() override
diff --git a/Engine/Source/Editor/Sequencer/Private/SequencerModule.cpp b/Engine/Source/Editor/Sequencer/Private/SequencerModule.cpp
--- a/Engine/Source/Editor/Sequencer/Private/SequencerModule.cpp
+++ b/Engine/Source/Editor/Sequencer/Private/SequencerModule.cpp
@@ -363,40 +363,39 @@ public:
 
 			if (Context->SelectedAssets.Num() == 1 && Context->SelectedAssets[0].IsInstanceOf(ULevelSequence::StaticClass()))
 			{
-				if (ULevelSequence* LevelSequence = Cast<ULevelSequence>(Context->SelectedAssets[0].GetAsset()))
+				const FAssetData LevelSequenceAsset = Context->SelectedAssets[0];
+			
+				// if this LevelSequence has associated maps, offer to load them
+				TArray<FString> AssociatedMaps = FSequencerUtilities::GetAssociatedLevelSequenceMapPackages(LevelSequenceAsset.PackageName);
+
+				if (AssociatedMaps.Num() > 0)
 				{
-					// if this LevelSequence has associated maps, offer to load them
-					TArray<FString> AssociatedMaps = FSequencerUtilities::GetAssociatedMapPackages(LevelSequence);
-
-					if(AssociatedMaps.Num()>0)
-					{
-						InSection.AddSubMenu(
-							"SequencerOpenMap_Label",
-							LOCTEXT("SequencerOpenMap_Label", "Open Map"),
-							LOCTEXT("SequencerOpenMap_Tooltip", "Open a map associated with this Level Sequence Asset"),
-							FNewMenuDelegate::CreateLambda(
-								[AssociatedMaps](FMenuBuilder& SubMenuBuilder)
+					InSection.AddSubMenu(
+						"SequencerOpenMap_Label",
+						LOCTEXT("SequencerOpenMap_Label", "Open Map"),
+						LOCTEXT("SequencerOpenMap_Tooltip", "Open a map associated with this Level Sequence Asset"),
+						FNewMenuDelegate::CreateLambda(
+							[AssociatedMaps](FMenuBuilder& SubMenuBuilder)
+							{
+								for (const FString& AssociatedMap : AssociatedMaps)
 								{
-									for (const FString& AssociatedMap : AssociatedMaps)
-									{
-										SubMenuBuilder.AddMenuEntry(
-											FText::FromString(FPaths::GetBaseFilename(AssociatedMap)),
-											FText(),
-											FSlateIcon(FAppStyle::GetAppStyleSetName(), "LevelEditor.Tabs.Levels"),
-											FExecuteAction::CreateLambda(
-												[AssociatedMap]
-												{
-													FEditorFileUtils::LoadMap(AssociatedMap);
-												}
-											)
-										);
-									}
+									SubMenuBuilder.AddMenuEntry(
+										FText::FromString(FPaths::GetBaseFilename(AssociatedMap)),
+										FText(),
+										FSlateIcon(FAppStyle::GetAppStyleSetName(), "LevelEditor.Tabs.Levels"),
+										FExecuteAction::CreateLambda(
+											[AssociatedMap]
+											{
+												FEditorFileUtils::LoadMap(AssociatedMap);
+											}
+										)
+									);
 								}
-							),
-							false,
-							FSlateIcon(FAppStyle::GetAppStyleSetName(), "LevelEditor.Tabs.Levels")
-						);
-					}
+							}
+						),
+						false,
+						FSlateIcon(FAppStyle::GetAppStyleSetName(), "LevelEditor.Tabs.Levels")
+					);
 				}
 			}
 		}));
diff --git a/Engine/Source/Editor/Sequencer/Private/SequencerUtilities.cpp b/Engine/Source/Editor/Sequencer/Private/SequencerUtilities.cpp
--- a/Engine/Source/Editor/Sequencer/Private/SequencerUtilities.cpp
+++ b/Engine/Source/Editor/Sequencer/Private/SequencerUtilities.cpp
@@ -411,22 +411,27 @@ FName FSequencerUtilities::GetUniqueName( FName CandidateName, const TArray<FNam
 	return UniqueName;
 }
 
-TArray<FString> FSequencerUtilities::GetAssociatedMapPackages(const ULevelSequence* InSequence)
+TArray<FString> FSequencerUtilities::GetAssociatedLevelSequenceMapPackages(const ULevelSequence* InSequence)
 {
 	if (!InSequence)
 	{
 		return TArray<FString>();
 	}
 
-	FAssetRegistryModule& AssetRegistryModule = FModuleManager::LoadModuleChecked<FAssetRegistryModule>("AssetRegistry");
 	const FName LSMapPathName = *InSequence->GetOutermost()->GetPathName();
+	return GetAssociatedLevelSequenceMapPackages(LSMapPathName);
+}
+
+TArray<FString> FSequencerUtilities::GetAssociatedLevelSequenceMapPackages(FName LevelSequencePackageName)
+{
+	FAssetRegistryModule& AssetRegistryModule = FModuleManager::LoadModuleChecked<FAssetRegistryModule>("AssetRegistry");
 
 	TArray<FString> AssociatedMaps;
 	TArray<FAssetIdentifier> AssociatedAssets;
 
 	// This makes the assumption these functions will append the array, and not clear it.
-	AssetRegistryModule.Get().GetReferencers(LSMapPathName, AssociatedAssets);
-	AssetRegistryModule.Get().GetDependencies(LSMapPathName, AssociatedAssets);
+	AssetRegistryModule.Get().GetReferencers(LevelSequencePackageName, AssociatedAssets);
+	AssetRegistryModule.Get().GetDependencies(LevelSequencePackageName, AssociatedAssets);
 
 	for (FAssetIdentifier& AssociatedMap : AssociatedAssets)
 	{
diff --git a/Engine/Source/Editor/Sequencer/Public/SequencerUtilities.h b/Engine/Source/Editor/Sequencer/Public/SequencerUtilities.h
--- a/Engine/Source/Editor/Sequencer/Public/SequencerUtilities.h
+++ b/Engine/Source/Editor/Sequencer/Public/SequencerUtilities.h
@@ -132,7 +132,8 @@ struct SEQUENCER_API FSequencerUtilities
 
 	static void PopulateMenu_SetBlendType(FMenuBuilder& MenuBuilder, const TArray<TWeakObjectPtr<UMovieSceneSection>>& InSections, TWeakPtr<ISequencer> InSequencer);
 
-	static TArray<FString> GetAssociatedMapPackages(const ULevelSequence* InSequence);
+	static TArray<FString> GetAssociatedLevelSequenceMapPackages(const ULevelSequence* InSequence);
+	static TArray<FString> GetAssociatedLevelSequenceMapPackages(FName LevelSequencePackageName);
 
 	/** 
 	 * Generates a unique FName from a candidate name given a set of already existing names.  
