Added cache track api to Niagara cache tracks to allow re-caching them from sequencer.

diff --git a/Engine/Plugins/FX/Niagara/Source/Niagara/Classes/NiagaraSimCache.h b/Engine/Plugins/FX/Niagara/Source/Niagara/Classes/NiagaraSimCache.h
--- a/Engine/Plugins/FX/Niagara/Source/Niagara/Classes/NiagaraSimCache.h
+++ b/Engine/Plugins/FX/Niagara/Source/Niagara/Classes/NiagaraSimCache.h
@@ -250,6 +250,12 @@ struct FNiagaraSimCacheLayout
 	TArray<FNiagaraSimCacheDataBuffersLayout> EmitterLayouts;
 };
 
+struct FNiagaraSimCacheWriteResult
+{
+	bool bSuccess = false;
+	FString ErrorMsg;
+};
+
 UCLASS(Experimental, BlueprintType)
 class NIAGARA_API UNiagaraSimCache : public UObject
 {
@@ -285,7 +291,7 @@ public:
 	ENiagaraSimCacheAttributeCaptureMode GetAttributeCaptureMode() const { return CreateParameters.AttributeCaptureMode; }
 
 	bool BeginWrite(FNiagaraSimCacheCreateParameters InCreateParameters, UNiagaraComponent* NiagaraComponent);
-	bool WriteFrame(UNiagaraComponent* NiagaraComponent);
+	FNiagaraSimCacheWriteResult WriteFrame(UNiagaraComponent* NiagaraComponent);
 	bool EndWrite();
 
 	bool CanRead(UNiagaraSystem* NiagaraSystem);
diff --git a/Engine/Plugins/FX/Niagara/Source/Niagara/Classes/NiagaraSimCacheCapture.cpp b/Engine/Plugins/FX/Niagara/Source/Niagara/Classes/NiagaraSimCacheCapture.cpp
--- a/Engine/Plugins/FX/Niagara/Source/Niagara/Classes/NiagaraSimCacheCapture.cpp
+++ b/Engine/Plugins/FX/Niagara/Source/Niagara/Classes/NiagaraSimCacheCapture.cpp
@@ -100,7 +100,7 @@ bool FNiagaraSimCacheCapture::OnFrameTick(float DeltaTime)
 	{
 		// If we fail to capture the frame it might be because things became invalid
 		// Or it might be because the simulation was not ticked since the last capture in which case don't advance the counter
-		if ( CaptureSimCache->WriteFrame(CaptureComponent) == false )
+		if ( CaptureSimCache->WriteFrame(CaptureComponent).bSuccess == false )
 		{
 			if ( CaptureSimCache->IsCacheValid() == false )
 			{
diff --git a/Engine/Plugins/FX/Niagara/Source/Niagara/Private/MovieScene/MovieSceneNiagaraSystemTrack.cpp b/Engine/Plugins/FX/Niagara/Source/Niagara/Private/MovieScene/MovieSceneNiagaraSystemTrack.cpp
--- a/Engine/Plugins/FX/Niagara/Source/Niagara/Private/MovieScene/MovieSceneNiagaraSystemTrack.cpp
+++ b/Engine/Plugins/FX/Niagara/Source/Niagara/Private/MovieScene/MovieSceneNiagaraSystemTrack.cpp
@@ -7,12 +7,10 @@
 #include "MovieScene/MovieSceneNiagaraSystemTrackTemplate.h"
 #include "Evaluation/MovieSceneEvalTemplate.h"
 #include "Evaluation/MovieSceneEvaluationTrack.h"
-#include "Compilation/IMovieSceneTemplateGenerator.h"
-#include "Compilation/MovieSceneSegmentCompiler.h"
-#include "Compilation/MovieSceneCompilerRules.h"
 
 #include UE_INLINE_GENERATED_CPP_BY_NAME(MovieSceneNiagaraSystemTrack)
 
+FMovieSceneSharedDataId UMovieSceneNiagaraSystemTrack::SharedDataId = FMovieSceneSharedDataId::Allocate();
 
 bool UMovieSceneNiagaraSystemTrack::SupportsType(TSubclassOf<UMovieSceneSection> SectionClass) const
 {
@@ -35,7 +33,6 @@ void UMovieSceneNiagaraSystemTrack::PostCompile(FMovieSceneEvaluationTrack& OutT
 	if (SpawnSectionPtr != nullptr)
 	{
 		UMovieSceneNiagaraSystemSpawnSection* SpawnSection = CastChecked<UMovieSceneNiagaraSystemSpawnSection>(*SpawnSectionPtr);
-		UMovieScene* ParentMovieScene = GetTypedOuter<UMovieScene>();
 		OutTrack.SetTrackImplementation(FMovieSceneNiagaraSystemTrackImplementation(
 			SpawnSection->GetInclusiveStartFrame(), SpawnSection->GetExclusiveEndFrame(),
 			SpawnSection->GetSectionStartBehavior(), SpawnSection->GetSectionEvaluateBehavior(),
diff --git a/Engine/Plugins/FX/Niagara/Source/Niagara/Private/MovieScene/MovieSceneNiagaraSystemTrackTemplate.cpp b/Engine/Plugins/FX/Niagara/Source/Niagara/Private/MovieScene/MovieSceneNiagaraSystemTrackTemplate.cpp
--- a/Engine/Plugins/FX/Niagara/Source/Niagara/Private/MovieScene/MovieSceneNiagaraSystemTrackTemplate.cpp
+++ b/Engine/Plugins/FX/Niagara/Source/Niagara/Private/MovieScene/MovieSceneNiagaraSystemTrackTemplate.cpp
@@ -9,6 +9,8 @@
 #include "Evaluation/MovieSceneEvaluationTemplateInstance.h"
 #include "MovieSceneSequence.h"
 #include "MovieScene.h"
+#include "Evaluation/MovieSceneEvaluationTrack.h"
+#include "MovieScene/MovieSceneNiagaraSystemTrack.h"
 
 #include UE_INLINE_GENERATED_CPP_BY_NAME(MovieSceneNiagaraSystemTrackTemplate)
 
@@ -280,9 +282,20 @@ FMovieSceneNiagaraSystemTrackImplementation::FMovieSceneNiagaraSystemTrackImplem
 
 void FMovieSceneNiagaraSystemTrackImplementation::Evaluate(const FMovieSceneEvaluationTrack& Track, TArrayView<const FMovieSceneFieldEntry_ChildTemplate> Children, const FMovieSceneEvaluationOperand& Operand, const FMovieSceneContext& Context, const FPersistentEvaluationData& PersistentData, FMovieSceneExecutionTokens& ExecutionTokens) const
 {
+	if (Track.GetSourceTrack()->IsEvalDisabled())
+	{
+		return;
+	}
 	ExecutionTokens.SetContext(Context);
-	ExecutionTokens.Add(FNiagaraSystemUpdateDesiredAgeExecutionToken(
+
+	// only add a token if there isn't one already, otherwise another track's token takes precendence
+	FMovieSceneSharedDataId TokenID = UMovieSceneNiagaraSystemTrack::SharedDataId;
+	FNiagaraSharedMarkerToken* SharedCacheToken = static_cast<FNiagaraSharedMarkerToken*>(ExecutionTokens.FindShared(TokenID));
+	if (SharedCacheToken == nullptr || !SharedCacheToken->BoundObjectIDs.Contains(Operand.ObjectBindingID))
+	{
+		ExecutionTokens.Add(FNiagaraSystemUpdateDesiredAgeExecutionToken(
 		SpawnSectionStartFrame, SpawnSectionEndFrame,
 		SpawnSectionStartBehavior, SpawnSectionEvaluateBehavior,
 		SpawnSectionEndBehavior, AgeUpdateMode, bAllowScalability));
+	}
 }
diff --git a/Engine/Plugins/FX/Niagara/Source/Niagara/Private/NiagaraSimCache.cpp b/Engine/Plugins/FX/Niagara/Source/Niagara/Private/NiagaraSimCache.cpp
--- a/Engine/Plugins/FX/Niagara/Source/Niagara/Private/NiagaraSimCache.cpp
+++ b/Engine/Plugins/FX/Niagara/Source/Niagara/Private/NiagaraSimCache.cpp
@@ -2,14 +2,11 @@
 
 #include "NiagaraSimCache.h"
 
-#include "NiagaraClearCounts.h"
 #include "NiagaraConstants.h"
 #include "NiagaraComponent.h"
-#include "NiagaraComputeExecutionContext.h"
 #include "NiagaraDataInterfaceUtilities.h"
 #include "NiagaraDataSetReadback.h"
 #include "NiagaraEmitterInstance.h"
-#include "NiagaraGPUInstanceCountManager.h"
 #include "NiagaraGpuComputeDispatchInterface.h"
 #include "NiagaraSimCacheAttributeReaderHelper.h"
 #include "NiagaraSimCacheCustomStorageInterface.h"
@@ -168,31 +165,36 @@ bool UNiagaraSimCache::BeginWrite(FNiagaraSimCacheCreateParameters InCreateParam
 	return true;
 }
 
-bool UNiagaraSimCache::WriteFrame(UNiagaraComponent* NiagaraComponent)
+FNiagaraSimCacheWriteResult UNiagaraSimCache::WriteFrame(UNiagaraComponent* NiagaraComponent)
 {
+	FNiagaraSimCacheWriteResult Result;
 	FNiagaraSimCacheHelper Helper(NiagaraComponent);
 	if (Helper.HasValidSimulationData() == false)
 	{
 		SoftNiagaraSystem.Reset();
-		return false;
+		Result.ErrorMsg = FString::Printf(TEXT("No valid simulation data for %s"), *GetPathNameSafe(Helper.NiagaraSystem));
+		return Result;
 	}
 
 	if ( SoftNiagaraSystem.Get() != Helper.NiagaraSystem )
 	{
+		Result.ErrorMsg = FString::Printf(TEXT("System %s != System %s"), *GetPathNameSafe(SoftNiagaraSystem.Get()), *GetPathNameSafe(Helper.NiagaraSystem));
 		SoftNiagaraSystem.Reset();
-		return false;
+		return Result;
 	}
 
 	// Simulation is complete nothing to cache
 	if ( Helper.SystemInstance->IsComplete() )
 	{
-		return false;
+		Result.ErrorMsg = FString::Printf(TEXT("System already completed its simulation. %s"), *Helper.SystemInstance->GetCrashReporterTag());
+		return Result;
 	}
 
 	// Is the simulation running?  If not nothing to cache yet
 	if ( Helper.SystemInstance->SystemInstanceState != ENiagaraSystemInstanceState::Running )
 	{
-		return false;
+		Result.ErrorMsg = FString::Printf(TEXT("System is not running, so there is no data to cache. %s"), *Helper.SystemInstance->GetCrashReporterTag());
+		return Result;
 	}
 
 	// First frame we are about to cache?
@@ -204,14 +206,16 @@ bool UNiagaraSimCache::WriteFrame(UNiagaraComponent* NiagaraComponent)
 	// If our tick counter hasn't moved then we won't capture a frame as there's nothing new to process
 	else if (CaptureTickCount == Helper.SystemInstance->GetTickCount())
 	{
-		return false;
+		Result.ErrorMsg = FString::Printf(TEXT("System was not ticked since the last capture. %s"), *Helper.SystemInstance->GetCrashReporterTag());
+		return Result;
 	}
 
 	// If the tick counter is lower than the previous value then the system was reset so we won't capture
 	if ( Helper.SystemInstance->GetTickCount() < CaptureTickCount )
 	{
+		Result.ErrorMsg = FString::Printf(TEXT("System was was reset since the last capture, skipping further cache writes. %s"), *Helper.SystemInstance->GetCrashReporterTag());
 		SoftNiagaraSystem.Reset();
-		return false;
+		return Result;
 	}
 
 	DurationSeconds = Helper.SystemInstance->GetAge() - StartSeconds;
@@ -275,6 +279,7 @@ bool UNiagaraSimCache::WriteFrame(UNiagaraComponent* NiagaraComponent)
 	{
 		const int FrameIndex = CacheFrames.Num() - 1;
 		bool bDataInterfacesSucess = true;
+		FString DataInterfaceName;
 
 		FNiagaraDataInterfaceUtilities::ForEachDataInterface(
 			Helper.SystemInstance,
@@ -284,7 +289,11 @@ bool UNiagaraSimCache::WriteFrame(UNiagaraComponent* NiagaraComponent)
 				{
 					INiagaraSimCacheCustomStorageInterface* SimCacheCustomStorageInterface = CastChecked<INiagaraSimCacheCustomStorageInterface>(DataInterface);
 					const void* PerInstanceData = Helper.SystemInstance->FindDataInterfaceInstanceData(DataInterface);
-					bDataInterfacesSucess &= SimCacheCustomStorageInterface->SimCacheWriteFrame(StorageObject, FrameIndex, Helper.SystemInstance, PerInstanceData);
+					if (!SimCacheCustomStorageInterface->SimCacheWriteFrame(StorageObject, FrameIndex, Helper.SystemInstance, PerInstanceData))
+					{
+						bDataInterfacesSucess = false;
+						DataInterfaceName += FString::Printf(TEXT("%s "), *DataInterface->GetName());
+					}
 				}
 				return true;
 			}
@@ -293,11 +302,13 @@ bool UNiagaraSimCache::WriteFrame(UNiagaraComponent* NiagaraComponent)
 		// A data interface failed to write information
 		if (bDataInterfacesSucess == false)
 		{
+			Result.ErrorMsg = FString::Printf(TEXT("Data interface(s) %s failed to write information for system %s"), *DataInterfaceName, *Helper.SystemInstance->GetCrashReporterTag());
 			SoftNiagaraSystem.Reset();
-			return false;
+			return Result;
 		}
 	}
-	return true;
+	Result.bSuccess = true;
+	return Result;
 }
 
 bool UNiagaraSimCache::EndWrite()
diff --git a/Engine/Plugins/FX/Niagara/Source/Niagara/Public/MovieScene/MovieSceneNiagaraSystemTrack.h b/Engine/Plugins/FX/Niagara/Source/Niagara/Public/MovieScene/MovieSceneNiagaraSystemTrack.h
--- a/Engine/Plugins/FX/Niagara/Source/Niagara/Public/MovieScene/MovieSceneNiagaraSystemTrack.h
+++ b/Engine/Plugins/FX/Niagara/Source/Niagara/Public/MovieScene/MovieSceneNiagaraSystemTrack.h
@@ -1,10 +1,19 @@
 // Copyright Epic Games, Inc. All Rights Reserved.
 
 #pragma once
+#include "MovieSceneExecutionToken.h"
 #include "MovieScene/MovieSceneNiagaraTrack.h"
 #include "Compilation/IMovieSceneTrackTemplateProducer.h"
+#include "Evaluation/PersistentEvaluationData.h"
 #include "MovieSceneNiagaraSystemTrack.generated.h"
 
+struct FNiagaraSharedMarkerToken : IMovieSceneSharedExecutionToken
+{
+	virtual void Execute(FPersistentEvaluationData& PersistentData, IMovieScenePlayer& Player) override {}
+
+	TSet<FGuid> BoundObjectIDs;
+};
+
 UCLASS(MinimalAPI)
 class UMovieSceneNiagaraSystemTrack : public UMovieSceneNiagaraTrack, public IMovieSceneTrackTemplateProducer
 {
@@ -16,4 +25,6 @@ public:
 	virtual FMovieSceneEvalTemplatePtr CreateTemplateForSection(const UMovieSceneSection& InSection) const override;
 	virtual void PostCompile(FMovieSceneEvaluationTrack& OutTrack, const FMovieSceneTrackCompilerArgs& Args) const override;
 	virtual bool PopulateEvaluationTree(TMovieSceneEvaluationTree<FMovieSceneTrackEvaluationData>& OutData) const override;
+
+	NIAGARA_API static FMovieSceneSharedDataId SharedDataId;
 };
\ No newline at end of file
diff --git a/Engine/Plugins/FX/NiagaraSimCaching/Source/NiagaraSimCaching/NiagaraSimCaching.Build.cs b/Engine/Plugins/FX/NiagaraSimCaching/Source/NiagaraSimCaching/NiagaraSimCaching.Build.cs
--- a/Engine/Plugins/FX/NiagaraSimCaching/Source/NiagaraSimCaching/NiagaraSimCaching.Build.cs
+++ b/Engine/Plugins/FX/NiagaraSimCaching/Source/NiagaraSimCaching/NiagaraSimCaching.Build.cs
@@ -20,7 +20,7 @@ namespace UnrealBuildTool.Rules
 					"Niagara",
 					"MovieScene",
 					"MovieSceneTracks",
-					"LevelSequence",
+					"LevelSequence", 
 				});
 
 			if(Target.bBuildEditor)
@@ -31,7 +31,8 @@ namespace UnrealBuildTool.Rules
 					{
 						"UnrealEd",
 						"SlateCore",
-						"Slate"
+						"Slate",
+						"TakeRecorder"
 					});
             }
 		}
diff --git a/Engine/Plugins/FX/NiagaraSimCaching/Source/NiagaraSimCaching/Private/Niagara/Sequencer/MovieSceneNiagaraCacheSection.cpp b/Engine/Plugins/FX/NiagaraSimCaching/Source/NiagaraSimCaching/Private/Niagara/Sequencer/MovieSceneNiagaraCacheSection.cpp
--- a/Engine/Plugins/FX/NiagaraSimCaching/Source/NiagaraSimCaching/Private/Niagara/Sequencer/MovieSceneNiagaraCacheSection.cpp
+++ b/Engine/Plugins/FX/NiagaraSimCaching/Source/NiagaraSimCaching/Private/Niagara/Sequencer/MovieSceneNiagaraCacheSection.cpp
@@ -1,7 +1,6 @@
 ﻿// Copyright Epic Games, Inc. All Rights Reserved.
 
 #include "Niagara/Sequencer/MovieSceneNiagaraCacheSection.h"
-#include "MovieScene.h"
 
 FMovieSceneNiagaraCacheParams::FMovieSceneNiagaraCacheParams() : FMovieSceneBaseCacheParams()
 {
@@ -17,6 +16,34 @@ float FMovieSceneNiagaraCacheParams::GetSequenceLength() const
 	return 0.f;
 }
 
+UMovieSceneNiagaraCacheSection::UMovieSceneNiagaraCacheSection()
+{
+	EvalOptions.CompletionMode = EMovieSceneCompletionMode::RestoreState;
+}
+
+#if WITH_EDITOR
+void UMovieSceneNiagaraCacheSection::PostEditChangeProperty(FPropertyChangedEvent& PropertyChangedEvent)
+{
+	if (PropertyChangedEvent.Property &&
+		PropertyChangedEvent.MemberProperty->GetFName() == GET_MEMBER_NAME_CHECKED(UMovieSceneNiagaraCacheSection, Params) &&
+		PropertyChangedEvent.Property->GetFName() == GET_MEMBER_NAME_CHECKED(FMovieSceneNiagaraCacheParams, SimCache))
+	{
+		if (TOptional<TRange<FFrameNumber>> AutoSizeRange = GetAutoSizeRange(); AutoSizeRange.IsSet())
+		{
+			SetRange(AutoSizeRange.GetValue());
+		}
+		bCacheOutOfDate = false;
+	}
+
+	if (PropertyChangedEvent.Property &&
+		PropertyChangedEvent.MemberProperty->GetFName() == GET_MEMBER_NAME_CHECKED(UMovieSceneNiagaraCacheSection, Params) &&
+		PropertyChangedEvent.Property->GetFName() == GET_MEMBER_NAME_CHECKED(FMovieSceneNiagaraCacheParams, CacheParameters))
+	{
+		bCacheOutOfDate = true;
+	}
+}
+#endif
+
 UMovieSceneNiagaraCacheSection::UMovieSceneNiagaraCacheSection(const FObjectInitializer& ObjectInitializer)
 	: Super(ObjectInitializer)
 {
diff --git a/Engine/Plugins/FX/NiagaraSimCaching/Source/NiagaraSimCaching/Private/Niagara/Sequencer/MovieSceneNiagaraCacheTemplate.cpp b/Engine/Plugins/FX/NiagaraSimCaching/Source/NiagaraSimCaching/Private/Niagara/Sequencer/MovieSceneNiagaraCacheTemplate.cpp
--- a/Engine/Plugins/FX/NiagaraSimCaching/Source/NiagaraSimCaching/Private/Niagara/Sequencer/MovieSceneNiagaraCacheTemplate.cpp
+++ b/Engine/Plugins/FX/NiagaraSimCaching/Source/NiagaraSimCaching/Private/Niagara/Sequencer/MovieSceneNiagaraCacheTemplate.cpp
@@ -2,7 +2,16 @@
 
 #include "MovieSceneNiagaraCacheTemplate.h"
 
+#include "IMovieScenePlayer.h"
+#include "MovieSceneExecutionToken.h"
 #include "NiagaraComponent.h"
+#include "Evaluation/MovieSceneEvaluationTrack.h"
+#include "Evaluation/MovieSceneExecutionTokens.h"
+#include "Misc/CoreMiscDefines.h"
+#include "Misc/CoreMiscDefines.h"
+#include "Misc/CoreMiscDefines.h"
+#include "MovieScene/MovieSceneNiagaraSystemTrack.h"
+#include "Niagara/Sequencer/MovieSceneNiagaraCacheTrack.h"
 
 DECLARE_CYCLE_STAT(TEXT("Niagara Cache - Token Execute"), MovieSceneEval_NiagaraCache_TokenExecute, STATGROUP_MovieSceneEval);
 
@@ -39,51 +48,170 @@ struct FPreAnimatedNiagaraCacheTokenProducer : IMovieScenePreAnimatedTokenProduc
 /** A movie scene execution token that executes a NiagaraCache */
 struct FNiagaraCacheExecutionToken : IMovieSceneExecutionToken
 {
-	FNiagaraCacheExecutionToken(const FMovieSceneNiagaraCacheSectionTemplateParameters &InParams) :
-		Params(InParams)
+	FNiagaraCacheExecutionToken(const TArray<FMovieSceneNiagaraSectionTemplateParameter> &InSections) :
+		CacheSections(InSections)
 	{}
 
 	/** Execute this token, operating on all objects referenced by 'Operand' */
 	virtual void Execute(const FMovieSceneContext& Context, const FMovieSceneEvaluationOperand& Operand, FPersistentEvaluationData& PersistentData, IMovieScenePlayer& Player) override
 	{
 		MOVIESCENE_DETAILED_SCOPE_CYCLE_COUNTER(MovieSceneEval_NiagaraCache_TokenExecute)
-		if (Operand.ObjectBindingID.IsValid() && Params.NiagaraCacheParams.SimCache && Params.NiagaraCacheParams.SimCache->GetNumFrames() > 0)
+
+		UNiagaraComponent* NiagaraComponent = nullptr;
+		if (Operand.ObjectBindingID.IsValid())
 		{
 			for (TWeakObjectPtr<> WeakObj : Player.FindBoundObjects(Operand))
 			{
-				if (UNiagaraComponent* NiagaraComponent = Cast<UNiagaraComponent>(WeakObj))
+				NiagaraComponent = Cast<UNiagaraComponent>(WeakObj);
+				if (NiagaraComponent)
 				{
-					if (NiagaraComponent->GetSimCache() != Params.NiagaraCacheParams.SimCache)
-					{
-						NiagaraComponent->SetSimCache(Params.NiagaraCacheParams.SimCache);
-					}
-					
-					// calculate the time at which to evaluate the sim cache
-					float Age = Params.MapTimeToAnimation(Params.NiagaraCacheParams, Params.NiagaraCacheParams.SimCache->GetDurationSeconds(), Context.GetTime(), Context.GetFrameRate()) + Params.NiagaraCacheParams.SimCache->GetStartSeconds();
-
-					Player.SavePreAnimatedState(*NiagaraComponent, FPreAnimatedNiagaraCacheTokenProducer::GetAnimTypeID(), FPreAnimatedNiagaraCacheTokenProducer());
-					NiagaraComponent->SetDesiredAge(Age);
-					NiagaraComponent->SetAgeUpdateMode(ENiagaraAgeUpdateMode::DesiredAge);
-					if (!NiagaraComponent->IsActive())
-					{
-						NiagaraComponent->ResetSystem();
-					}
+					break;
 				}
 			}
 		}
+		if (!NiagaraComponent)
+		{
+			return;
+		}
+		
+		ENiagaraSimCacheSectionPlayMode ReplayMode = ENiagaraSimCacheSectionPlayMode::DisplayCacheOnly;
+		for (const FMovieSceneNiagaraSectionTemplateParameter& Section : CacheSections)
+		{
+			ReplayMode = Section.Params.CacheReplayPlayMode;
+			if (Section.IsTimeWithinSection(Context.GetTime().GetFrame()) && Section.Params.SimCache)
+			{
+				if (NiagaraComponent->GetSimCache() != Section.Params.SimCache)
+				{
+					NiagaraComponent->SetSimCache(Section.Params.SimCache);
+				}
+		
+				// calculate the time at which to evaluate the sim cache
+				float Age = MapTimeToAnimation(Section.Params, Section.Params.SimCache->GetDurationSeconds(), Context.GetTime(), Context.GetFrameRate(), Section) + Section.Params.SimCache->GetStartSeconds();
+
+				FScopedPreAnimatedCaptureSource CaptureSource(&Player.PreAnimatedState, PersistentData.GetSectionKey(), true);
+				Player.SavePreAnimatedState(*NiagaraComponent, FPreAnimatedNiagaraCacheTokenProducer::GetAnimTypeID(), FPreAnimatedNiagaraCacheTokenProducer());
+				NiagaraComponent->SetDesiredAge(Age);
+				NiagaraComponent->SetAgeUpdateMode(ENiagaraAgeUpdateMode::DesiredAge);
+				if (!NiagaraComponent->IsActive())
+				{
+					NiagaraComponent->ResetSystem();
+				}
+				return;
+			}
+		}
+
+		// outside of cached sections, use configured replay mode
+		if (ReplayMode == ENiagaraSimCacheSectionPlayMode::DisplayCacheOnly)
+		{
+			NiagaraComponent->DeactivateImmediate();
+		}
+		else if (ReplayMode == ENiagaraSimCacheSectionPlayMode::SimWithoutCache)
+		{
+			NiagaraComponent->SetSimCache(nullptr);
+		}
+	}
+
+	float MapTimeToAnimation(const FMovieSceneBaseCacheParams& BaseParams, float ComponentDuration, FFrameTime InPosition, FFrameRate InFrameRate, const FMovieSceneNiagaraSectionTemplateParameter& Section) const
+	{
+		FFrameNumber SectionStartTime = Section.GetInclusiveStartFrame();
+		FFrameNumber SectionEndTime = Section.GetExclusiveEndFrame();
+		
+		const float SequenceLength = Section.Params.SectionStretchMode == ENiagaraSimCacheSectionStretchMode::TimeDilate ? InFrameRate.AsSeconds(SectionEndTime - SectionStartTime) : ComponentDuration;
+		const FFrameTime AnimationLength = SequenceLength * InFrameRate;
+		const int32 LengthInFrames = AnimationLength.FrameNumber.Value + static_cast<int>(AnimationLength.GetSubFrame() + 0.5f) + 1;
+	
+		//we only play end if we are not looping, and assuming we are looping if Length is greater than default length;
+		const bool bLooping = (SectionEndTime.Value - SectionStartTime.Value + BaseParams.StartFrameOffset + BaseParams.EndFrameOffset) > LengthInFrames;
+
+		InPosition = FMath::Clamp(InPosition, FFrameTime(SectionStartTime), FFrameTime(SectionEndTime - 1));
+
+		const float SectionPlayRate = BaseParams.PlayRate;
+		const float AnimPlayRate = FMath::IsNearlyZero(SectionPlayRate) ? 1.0f : SectionPlayRate;
+		const float SeqLength = SequenceLength - InFrameRate.AsSeconds(BaseParams.StartFrameOffset + BaseParams.EndFrameOffset);
+
+		float AnimPosition = FFrameTime::FromDecimal((InPosition - SectionStartTime).AsDecimal() * AnimPlayRate) / InFrameRate;
+		AnimPosition += InFrameRate.AsSeconds(BaseParams.FirstLoopStartFrameOffset);
+		if (SeqLength > 0.f && (bLooping || !FMath::IsNearlyEqual(AnimPosition, SeqLength, 1e-4f)))
+		{
+			AnimPosition = FMath::Fmod(AnimPosition, SeqLength);
+		}
+		AnimPosition += InFrameRate.AsSeconds(BaseParams.StartFrameOffset);
+		if (BaseParams.bReverse)
+		{
+			AnimPosition = SequenceLength - AnimPosition;
+		}
+		if (Section.Params.SectionStretchMode == ENiagaraSimCacheSectionStretchMode::TimeDilate)
+		{
+			AnimPosition = (AnimPosition / SequenceLength) * ComponentDuration;
+		}
+
+		return AnimPosition;
 	}
 
-	FMovieSceneNiagaraCacheSectionTemplateParameters Params;
+	TArray<FMovieSceneNiagaraSectionTemplateParameter> CacheSections;
 };
 
-FMovieSceneNiagaraCacheSectionTemplate::FMovieSceneNiagaraCacheSectionTemplate(const UMovieSceneNiagaraCacheSection& InSection)
-	: Params(const_cast<FMovieSceneNiagaraCacheParams&> (InSection.Params), InSection.GetInclusiveStartFrame(), InSection.GetExclusiveEndFrame())
+struct FNiagaraCacheResetToken : IMovieSceneExecutionToken
+{
+	virtual void Execute(const FMovieSceneContext& Context, const FMovieSceneEvaluationOperand& Operand, FPersistentEvaluationData& PersistentData, IMovieScenePlayer& Player) override
+	{
+		MOVIESCENE_DETAILED_SCOPE_CYCLE_COUNTER(MovieSceneEval_NiagaraCache_TokenExecute)
+
+		UNiagaraComponent* NiagaraComponent = nullptr;
+		if (Operand.ObjectBindingID.IsValid())
+		{
+			for (TWeakObjectPtr<> WeakObj : Player.FindBoundObjects(Operand))
+			{
+				NiagaraComponent = Cast<UNiagaraComponent>(WeakObj);
+				if (NiagaraComponent)
+				{
+					break;
+				}
+			}
+		}
+		if (!NiagaraComponent)
+		{
+			return;
+		}
+
+		NiagaraComponent->SetSimCache(nullptr);
+	}
+};
+
+FMovieSceneNiagaraCacheSectionTemplate::FMovieSceneNiagaraCacheSectionTemplate(TArray<FMovieSceneNiagaraSectionTemplateParameter> InCacheSections)
+	: CacheSections(InCacheSections)
 {
 }
 
-//We use a token here so we can set the manual tick state back to what it was previously when outside this section.
-//This is similar to how Skeletal Animation evaluation also works.
-void FMovieSceneNiagaraCacheSectionTemplate::Evaluate(const FMovieSceneEvaluationOperand& Operand, const FMovieSceneContext& Context, const FPersistentEvaluationData& PersistentData, FMovieSceneExecutionTokens& ExecutionTokens) const
+void FMovieSceneNiagaraCacheSectionTemplate::Evaluate(const FMovieSceneEvaluationTrack& Track, TArrayView<const FMovieSceneFieldEntry_ChildTemplate> Children, const FMovieSceneEvaluationOperand& Operand, const FMovieSceneContext& Context, const FPersistentEvaluationData& PersistentData,	FMovieSceneExecutionTokens& ExecutionTokens) const
 {
-	ExecutionTokens.Add(FNiagaraCacheExecutionToken(Params));
-}
\ No newline at end of file
+	UMovieSceneNiagaraCacheTrack* NiagaraTrack = Cast<UMovieSceneNiagaraCacheTrack>(Track.GetSourceTrack());
+	if (NiagaraTrack == nullptr || NiagaraTrack->bIsRecording)
+	{
+		return;
+	}
+	ExecutionTokens.SetContext(Context);
+
+	// remove the sim cache if we don't have section data to display
+	if (NiagaraTrack->IsEvalDisabled() || CacheSections.Num() == 0)
+	{
+		ExecutionTokens.Add(FNiagaraCacheResetToken());
+		return;
+	}
+	
+	// only add a token if there isn't one already, otherwise another track's token takes precendence
+	FMovieSceneSharedDataId TokenID = UMovieSceneNiagaraSystemTrack::SharedDataId;
+	FNiagaraSharedMarkerToken* SharedCacheToken = static_cast<FNiagaraSharedMarkerToken*>(ExecutionTokens.FindShared(TokenID));
+	if (SharedCacheToken == nullptr)
+	{
+		ExecutionTokens.Add(FNiagaraCacheExecutionToken(CacheSections));
+		FNiagaraSharedMarkerToken SharedToken = FNiagaraSharedMarkerToken();
+		SharedToken.BoundObjectIDs.Add(Operand.ObjectBindingID);
+		ExecutionTokens.AddShared(TokenID, SharedToken);
+	}
+	else if (!SharedCacheToken->BoundObjectIDs.Contains(Operand.ObjectBindingID))
+	{
+		ExecutionTokens.Add(FNiagaraCacheExecutionToken(CacheSections));
+		SharedCacheToken->BoundObjectIDs.Add(Operand.ObjectBindingID);
+	}
+}
diff --git a/Engine/Plugins/FX/NiagaraSimCaching/Source/NiagaraSimCaching/Private/Niagara/Sequencer/MovieSceneNiagaraCacheTemplate.h b/Engine/Plugins/FX/NiagaraSimCaching/Source/NiagaraSimCaching/Private/Niagara/Sequencer/MovieSceneNiagaraCacheTemplate.h
--- a/Engine/Plugins/FX/NiagaraSimCaching/Source/NiagaraSimCaching/Private/Niagara/Sequencer/MovieSceneNiagaraCacheTemplate.h
+++ b/Engine/Plugins/FX/NiagaraSimCaching/Source/NiagaraSimCaching/Private/Niagara/Sequencer/MovieSceneNiagaraCacheTemplate.h
@@ -2,37 +2,57 @@
 
 #pragma once
 
-#include "Evaluation/MovieSceneBaseCacheTemplate.h"
+#include "Evaluation/MovieSceneTrackImplementation.h"
 #include "Niagara/Sequencer/MovieSceneNiagaraCacheSection.h"
 #include "MovieSceneNiagaraCacheTemplate.generated.h"
 
 USTRUCT()
-struct FMovieSceneNiagaraCacheSectionTemplateParameters : public FMovieSceneBaseCacheSectionTemplateParameters
+struct FMovieSceneNiagaraSectionTemplateParameter
 {
 	GENERATED_BODY()
-
-	FMovieSceneNiagaraCacheSectionTemplateParameters() {}
-
-	FMovieSceneNiagaraCacheSectionTemplateParameters(const FMovieSceneNiagaraCacheParams& InNiagaraCacheParams, FFrameNumber InSectionStartTime, FFrameNumber InSectionEndTime)
-		: FMovieSceneBaseCacheSectionTemplateParameters(InSectionStartTime, InSectionEndTime), NiagaraCacheParams(InNiagaraCacheParams)
- 
-	{}
+	
+	UPROPERTY()
+	FMovieSceneFrameRange SectionRange;
 
 	UPROPERTY()
-	FMovieSceneNiagaraCacheParams NiagaraCacheParams;
+	FMovieSceneNiagaraCacheParams Params;
+	
+	bool IsTimeWithinSection(const FFrameNumber& Position) const 
+	{
+		return SectionRange.Value.Contains(Position);
+	}
+
+	FFrameNumber GetInclusiveStartFrame() const
+	{
+		TRangeBound<FFrameNumber> LowerBound = SectionRange.GetLowerBound();
+		return LowerBound.IsInclusive() ? LowerBound.GetValue() : LowerBound.GetValue() + 1;
+	}
+
+	FFrameNumber GetExclusiveEndFrame() const
+	{
+		TRangeBound<FFrameNumber> UpperBound = SectionRange.GetUpperBound();
+		return UpperBound.IsInclusive() ? UpperBound.GetValue() + 1 : UpperBound.GetValue();
+	}
 };
 
 USTRUCT()
-struct FMovieSceneNiagaraCacheSectionTemplate : public FMovieSceneEvalTemplate
+struct FMovieSceneNiagaraCacheSectionTemplate : public FMovieSceneTrackImplementation
 {
 	GENERATED_BODY()
 
 	FMovieSceneNiagaraCacheSectionTemplate() {}
-	FMovieSceneNiagaraCacheSectionTemplate(const UMovieSceneNiagaraCacheSection& Section);
+	FMovieSceneNiagaraCacheSectionTemplate(TArray<FMovieSceneNiagaraSectionTemplateParameter> CacheSections);
 
+	virtual void SetupOverrides() override
+	{
+		EnableOverrides(CustomEvaluateFlag);
+	}
+
+private:
 	virtual UScriptStruct& GetScriptStructImpl() const override { return *StaticStruct(); }
-	virtual void Evaluate(const FMovieSceneEvaluationOperand& Operand, const FMovieSceneContext& Context, const FPersistentEvaluationData& PersistentData, FMovieSceneExecutionTokens& ExecutionTokens) const override;
+	virtual void Evaluate(const FMovieSceneEvaluationTrack& Track, TArrayView<const FMovieSceneFieldEntry_ChildTemplate> Children, const FMovieSceneEvaluationOperand& Operand, const FMovieSceneContext& Context, const FPersistentEvaluationData& PersistentData, FMovieSceneExecutionTokens& ExecutionTokens) const override;
+
 
 	UPROPERTY()
-	FMovieSceneNiagaraCacheSectionTemplateParameters Params;
+	TArray<FMovieSceneNiagaraSectionTemplateParameter> CacheSections;
 };
diff --git a/Engine/Plugins/FX/NiagaraSimCaching/Source/NiagaraSimCaching/Private/Niagara/Sequencer/MovieSceneNiagaraCacheTrack.cpp b/Engine/Plugins/FX/NiagaraSimCaching/Source/NiagaraSimCaching/Private/Niagara/Sequencer/MovieSceneNiagaraCacheTrack.cpp
--- a/Engine/Plugins/FX/NiagaraSimCaching/Source/NiagaraSimCaching/Private/Niagara/Sequencer/MovieSceneNiagaraCacheTrack.cpp
+++ b/Engine/Plugins/FX/NiagaraSimCaching/Source/NiagaraSimCaching/Private/Niagara/Sequencer/MovieSceneNiagaraCacheTrack.cpp
@@ -2,9 +2,11 @@
 
 #include "Niagara/Sequencer/MovieSceneNiagaraCacheTrack.h"
 #include "MovieScene.h"
-#include "Niagara/Sequencer/MovieSceneNiagaraCacheSection.h"
 #include "MovieSceneNiagaraCacheTemplate.h"
+#include "Niagara/Sequencer/MovieSceneNiagaraCacheSection.h"
 #include "NiagaraSimCache.h"
+#include "Evaluation/MovieSceneEvalTemplate.h"
+#include "Evaluation/MovieSceneEvaluationTrack.h"
 
 #define LOCTEXT_NAMESPACE "MovieSceneNiagaraCacheTrack"
 
@@ -57,7 +59,7 @@ TArray<UMovieSceneSection*> UMovieSceneNiagaraCacheTrack::GetAnimSectionsAtTime(
 
 FMovieSceneEvalTemplatePtr UMovieSceneNiagaraCacheTrack::CreateTemplateForSection(const UMovieSceneSection& InSection) const
 {
-	return FMovieSceneNiagaraCacheSectionTemplate(*CastChecked<UMovieSceneNiagaraCacheSection>(&InSection));
+	return FMovieSceneEvalTemplatePtr();
 }
 
 /* UMovieSceneTrack interface
@@ -81,6 +83,35 @@ UMovieSceneSection* UMovieSceneNiagaraCacheTrack::CreateNewSection()
 	return MovieSceneNiagaraCacheSection;
 }
 
+bool UMovieSceneNiagaraCacheTrack::PopulateEvaluationTree(TMovieSceneEvaluationTree<FMovieSceneTrackEvaluationData>& OutData) const
+{
+	UMovieSceneNiagaraCacheTrack* This = const_cast<UMovieSceneNiagaraCacheTrack*>(this);
+	// We always want to evaluate everything
+	OutData.Add(TRange<FFrameNumber>::All(), FMovieSceneTrackEvaluationData::FromTrack(This));
+
+	return true;
+}
+
+void UMovieSceneNiagaraCacheTrack::PostCompile(FMovieSceneEvaluationTrack& OutTrack, const FMovieSceneTrackCompilerArgs& Args) const
+{
+	TArray<FMovieSceneNiagaraSectionTemplateParameter> TemplateParams;
+	for (UMovieSceneSection* Section : AnimationSections)
+	{
+		if (UMovieSceneNiagaraCacheSection* NiagaraCacheSection = Cast<UMovieSceneNiagaraCacheSection>(Section))
+		{
+			if (NiagaraCacheSection->IsActive() && NiagaraCacheSection->Params.SimCache)
+			{
+				FMovieSceneNiagaraSectionTemplateParameter& Param = TemplateParams.AddDefaulted_GetRef();
+				Param.SectionRange = NiagaraCacheSection->SectionRange;
+				Param.Params = NiagaraCacheSection->Params;
+			}
+		}
+	}
+	
+	OutTrack.SetTrackImplementation(FMovieSceneNiagaraCacheSectionTemplate(TemplateParams));
+	OutTrack.SetEvaluationPriority(10000);
+}
+
 void UMovieSceneNiagaraCacheTrack::RemoveAllAnimationData()
 {
 	AnimationSections.Empty();
@@ -120,4 +151,29 @@ FText UMovieSceneNiagaraCacheTrack::GetDefaultDisplayName() const
 
 #endif
 
+void UMovieSceneNiagaraCacheTrack::ResetCache()
+{
+	for (TObjectPtr<UMovieSceneSection>& Section : AnimationSections)
+	{
+		if (UMovieSceneNiagaraCacheSection* NiagaraCacheSection = Cast<UMovieSceneNiagaraCacheSection>(Section))
+		{
+			if (UNiagaraSimCache* NiagaraSimCache = NiagaraCacheSection->Params.SimCache)
+			{
+				//TODO (mga) not sure if we need to do anything here, as the cache is cleared when BeginWrite is called on it
+				//NiagaraSimCache->ResetCache();
+			}
+		}
+	}
+}
+
+void UMovieSceneNiagaraCacheTrack::SetCacheRecordingAllowed(bool bShouldRecord)
+{
+	bCacheRecordingEnabled = bShouldRecord;
+}
+
+bool UMovieSceneNiagaraCacheTrack::IsCacheRecordingAllowed() const
+{
+	return bCacheRecordingEnabled;
+}
+
 #undef LOCTEXT_NAMESPACE
diff --git a/Engine/Plugins/FX/NiagaraSimCaching/Source/NiagaraSimCaching/Public/Niagara/Sequencer/MovieSceneNiagaraCacheSection.h b/Engine/Plugins/FX/NiagaraSimCaching/Source/NiagaraSimCaching/Public/Niagara/Sequencer/MovieSceneNiagaraCacheSection.h
--- a/Engine/Plugins/FX/NiagaraSimCaching/Source/NiagaraSimCaching/Public/Niagara/Sequencer/MovieSceneNiagaraCacheSection.h
+++ b/Engine/Plugins/FX/NiagaraSimCaching/Source/NiagaraSimCaching/Public/Niagara/Sequencer/MovieSceneNiagaraCacheSection.h
@@ -6,6 +6,34 @@
 #include "Sections/MovieSceneBaseCacheSection.h"
 #include "MovieSceneNiagaraCacheSection.generated.h"
 
+UENUM(BlueprintType)
+enum class ENiagaraSimCacheSectionPlayMode : uint8
+{
+	/**
+	When the sequence has no cached data to display, the Niagara component runs the simulation normally
+	*/
+	SimWithoutCache,
+
+	/**
+	When the sequence has no cached data to display, the Niagara component is disabled
+	*/
+	DisplayCacheOnly,
+};
+
+UENUM(BlueprintType)
+enum class ENiagaraSimCacheSectionStretchMode : uint8
+{
+	/**
+	When the cache section is stretched in the track it will repeat the cached data 
+	*/
+	Repeat,
+
+	/**
+	When the cache section is stretched in the track it will dilate the input time so the cached data is stretched once over the full section
+	*/
+	TimeDilate,
+};
+
 USTRUCT()
 struct NIAGARASIMCACHING_API FMovieSceneNiagaraCacheParams : public FMovieSceneBaseCacheParams
 {
@@ -17,12 +45,20 @@ struct NIAGARASIMCACHING_API FMovieSceneNiagaraCacheParams : public FMovieSceneB
 	/** Gets the animation sequence length, not modified by play rate */
 	virtual float GetSequenceLength() const override;
 
-	UPROPERTY(VisibleAnywhere, Category = "NiagaraCache")
+	UPROPERTY(EditAnywhere, Category = "NiagaraCache")
 	FNiagaraSimCacheCreateParameters CacheParameters;
 	
 	/** The sim cache this section plays and records into */
-	UPROPERTY()
+	UPROPERTY(EditAnywhere, Category = "NiagaraCache")
 	TObjectPtr<UNiagaraSimCache> SimCache;
+
+	/** What should the effect do when the track has no cache data to display */
+	UPROPERTY(EditAnywhere, Category="SimCache")
+	ENiagaraSimCacheSectionPlayMode CacheReplayPlayMode = ENiagaraSimCacheSectionPlayMode::DisplayCacheOnly;
+
+	/** What should the effect do when the cache section is stretched? */
+	UPROPERTY(EditAnywhere, Category="SimCache")
+	ENiagaraSimCacheSectionStretchMode SectionStretchMode = ENiagaraSimCacheSectionStretchMode::TimeDilate;
 };
 
 /**
@@ -35,6 +71,15 @@ class UMovieSceneNiagaraCacheSection : public UMovieSceneBaseCacheSection
 
 public:
 
+	UMovieSceneNiagaraCacheSection();
+
+#if WITH_EDITOR
+	virtual void PostEditChangeProperty(FPropertyChangedEvent& PropertyChangedEvent) override;
+#endif
+
 	UPROPERTY(EditAnywhere, Category = "NiagaraCache", meta = (ShowOnlyInnerProperties))
 	FMovieSceneNiagaraCacheParams Params;
+
+	UPROPERTY()
+	bool bCacheOutOfDate = false;
 };
diff --git a/Engine/Plugins/FX/NiagaraSimCaching/Source/NiagaraSimCaching/Public/Niagara/Sequencer/MovieSceneNiagaraCacheTrack.h b/Engine/Plugins/FX/NiagaraSimCaching/Source/NiagaraSimCaching/Public/Niagara/Sequencer/MovieSceneNiagaraCacheTrack.h
--- a/Engine/Plugins/FX/NiagaraSimCaching/Source/NiagaraSimCaching/Public/Niagara/Sequencer/MovieSceneNiagaraCacheTrack.h
+++ b/Engine/Plugins/FX/NiagaraSimCaching/Source/NiagaraSimCaching/Public/Niagara/Sequencer/MovieSceneNiagaraCacheTrack.h
@@ -4,6 +4,7 @@
 
 #include "MovieSceneNameableTrack.h"
 #include "Compilation/IMovieSceneTrackTemplateProducer.h"
+#include "Tracks/MovieSceneCachedTrack.h"
 #include "MovieSceneNiagaraCacheTrack.generated.h"
 
 /**
@@ -13,6 +14,7 @@ UCLASS(MinimalAPI)
 class UMovieSceneNiagaraCacheTrack
 	: public UMovieSceneNameableTrack
 	, public IMovieSceneTrackTemplateProducer
+	, public IMovieSceneCachedTrack
 {
 	GENERATED_UCLASS_BODY()
 
@@ -34,6 +36,8 @@ public:
 	virtual const TArray<UMovieSceneSection*>& GetAllSections() const override;
 	virtual bool SupportsType(TSubclassOf<UMovieSceneSection> SectionClass) const override;
 	virtual UMovieSceneSection* CreateNewSection() override;
+	virtual bool PopulateEvaluationTree(TMovieSceneEvaluationTree<FMovieSceneTrackEvaluationData>& OutData) const override;
+	virtual void PostCompile(FMovieSceneEvaluationTrack& OutTrack, const FMovieSceneTrackCompilerArgs& Args) const override;
 
 	// ~IMovieSceneTrackTemplateProducer interface
 	virtual FMovieSceneEvalTemplatePtr CreateTemplateForSection(const UMovieSceneSection& InSection) const override;
@@ -42,9 +46,20 @@ public:
 	virtual FText GetDefaultDisplayName() const override;
 #endif
 
+	// ~IMovieSceneCachedTrack interface
+	virtual void ResetCache() override;
+	virtual void SetCacheRecordingAllowed(bool bShouldRecord) override;
+	virtual bool IsCacheRecordingAllowed() const override;
+	
+	UPROPERTY(Transient)
+	bool bIsRecording = false;
+
 private:
 
 	/** List of all animation sections */
 	UPROPERTY()
 	TArray<TObjectPtr<UMovieSceneSection>> AnimationSections;
+
+	UPROPERTY()
+	bool bCacheRecordingEnabled = true;
 };
diff --git a/Engine/Plugins/FX/NiagaraSimCaching/Source/NiagaraSimCachingEditor/NiagaraSimCachingEditor.Build.cs b/Engine/Plugins/FX/NiagaraSimCaching/Source/NiagaraSimCachingEditor/NiagaraSimCachingEditor.Build.cs
--- a/Engine/Plugins/FX/NiagaraSimCaching/Source/NiagaraSimCachingEditor/NiagaraSimCachingEditor.Build.cs
+++ b/Engine/Plugins/FX/NiagaraSimCaching/Source/NiagaraSimCachingEditor/NiagaraSimCachingEditor.Build.cs
@@ -33,6 +33,7 @@ namespace UnrealBuildTool.Rules
 					"TakesCore",
 					"TakeRecorder",
 					"TakeTrackRecorders",
+					"CacheTrackRecorder",
 				});
 		}
 	}
diff --git a/Engine/Plugins/FX/NiagaraSimCaching/Source/NiagaraSimCachingEditor/Private/Niagara/NiagaraSimCachingEditorPlugin.cpp b/Engine/Plugins/FX/NiagaraSimCaching/Source/NiagaraSimCachingEditor/Private/Niagara/NiagaraSimCachingEditorPlugin.cpp
--- a/Engine/Plugins/FX/NiagaraSimCaching/Source/NiagaraSimCachingEditor/Private/Niagara/NiagaraSimCachingEditorPlugin.cpp
+++ b/Engine/Plugins/FX/NiagaraSimCaching/Source/NiagaraSimCachingEditor/Private/Niagara/NiagaraSimCachingEditorPlugin.cpp
@@ -2,22 +2,21 @@
 
 #include "Niagara/NiagaraSimCachingEditorPlugin.h"
 #include "Features/IModularFeatures.h"
-#include "Sequencer/TakeRecorderNiagaraCacheSource.h"
 #include "CoreMinimal.h"
-#include "Engine/Selection.h"
 #include "Modules/ModuleManager.h"
 #include "ISequencerModule.h"
 #include "Sequencer/NiagaraCacheTrackEditor.h"
-#include "ITakeRecorderModule.h"
 
 IMPLEMENT_MODULE(INiagaraSimCachingEditorPlugin, NiagaraSimCachingEditor)
 
+DEFINE_LOG_CATEGORY(LogNiagaraSimCachingEditor);
+
 #define LOCTEXT_NAMESPACE "CacheEditorPlugin"
 
 static const FName MovieSceneTrackRecorderFactoryName("MovieSceneTrackRecorderFactory");
 
 void INiagaraSimCachingEditorPlugin::StartupModule()
-{	
+{
 	ISequencerModule& SequencerModule = FModuleManager::Get().LoadModuleChecked<ISequencerModule>("Sequencer");
 	TrackEditorBindingHandle = SequencerModule.RegisterTrackEditor(FOnCreateTrackEditor::CreateStatic(&FNiagaraCacheTrackEditor::CreateTrackEditor));
 		
diff --git a/Engine/Plugins/FX/NiagaraSimCaching/Source/NiagaraSimCachingEditor/Private/Niagara/NiagaraSimCachingEditorStyle.cpp b/Engine/Plugins/FX/NiagaraSimCaching/Source/NiagaraSimCachingEditor/Private/Niagara/NiagaraSimCachingEditorStyle.cpp
--- a/Engine/Plugins/FX/NiagaraSimCaching/Source/NiagaraSimCachingEditor/Private/Niagara/NiagaraSimCachingEditorStyle.cpp
+++ b/Engine/Plugins/FX/NiagaraSimCaching/Source/NiagaraSimCachingEditor/Private/Niagara/NiagaraSimCachingEditorStyle.cpp
@@ -6,13 +6,15 @@
 
 FName FNiagaraSimCachingEditorStyle::StyleName("NiagaraSimCachingEditorStyle");
 
-FNiagaraSimCachingEditorStyle::FNiagaraSimCachingEditorStyle()
-	: FSlateStyleSet(StyleName)
+FNiagaraSimCachingEditorStyle::FNiagaraSimCachingEditorStyle() : FSlateStyleSet(StyleName)
 {
-	const FVector2D IconSize(20.0f, 20.0f);
-	const FVector2D LabelIconSize(16.0f, 16.0f);
+	FVector2D Icon16x16(16.0f, 16.0f);
+	
+	SetContentRoot(FPaths::EnginePluginsDir() / TEXT("FX/NiagaraSimCaching/Content"));
+	SetCoreContentRoot(FPaths::EngineContentDir() / TEXT("Slate"));
 
-	FSlateStyleSet::SetCoreContentRoot(FPaths::EngineContentDir() / TEXT("Editor/Slate"));
+	Set("Niagara.SimCaching.StatusIcon.Color", FLinearColor(0, 1.f, 1.f));
+	Set("Niagara.SimCaching.RecordIconSmall", new FSlateImageBrush(RootToContentDir(TEXT("RecordButton_Idle.png")), Icon16x16));
 
 	FSlateStyleRegistry::RegisterSlateStyle(*this);
 }
diff --git a/Engine/Plugins/FX/NiagaraSimCaching/Source/NiagaraSimCachingEditor/Private/Sequencer/MovieSceneNiagaraTrackRecorder.cpp b/Engine/Plugins/FX/NiagaraSimCaching/Source/NiagaraSimCachingEditor/Private/Sequencer/MovieSceneNiagaraTrackRecorder.cpp
--- a/Engine/Plugins/FX/NiagaraSimCaching/Source/NiagaraSimCachingEditor/Private/Sequencer/MovieSceneNiagaraTrackRecorder.cpp
+++ b/Engine/Plugins/FX/NiagaraSimCaching/Source/NiagaraSimCachingEditor/Private/Sequencer/MovieSceneNiagaraTrackRecorder.cpp
@@ -1,16 +1,19 @@
 ﻿// Copyright Epic Games, Inc. All Rights Reserved.
 
 #include "Sequencer/MovieSceneNiagaraTrackRecorder.h"
+
+#include "ISequencer.h"
+#include "LevelSequence.h"
+#include "NiagaraComponent.h"
+#include "TakeRecorderSettings.h"
+#include "Misc/App.h"
+#include "Misc/CoreMiscDefines.h"
+#include "Misc/CoreMiscDefines.h"
+#include "MovieScene/MovieSceneNiagaraSystemTrack.h"
+#include "Niagara/NiagaraSimCachingEditorPlugin.h"
 #include "Niagara/Sequencer/MovieSceneNiagaraCacheSection.h"
 #include "Niagara/Sequencer/MovieSceneNiagaraCacheTrack.h"
-#include "TakeRecorderSettings.h"
 #include "Recorder/TakeRecorderParameters.h"
-#include "Engine/TimecodeProvider.h"
-#include "Misc/App.h"
-#include "MovieSceneFolder.h"
-#include "LevelSequence.h"
-#include "NiagaraComponent.h"
-#include "NiagaraSimCacheFunctionLibrary.h"
 
 #define LOCTEXT_NAMESPACE "MovieSceneNiagaraTrackRecorder"
 
@@ -24,6 +27,76 @@ UMovieSceneTrackRecorder* FMovieSceneNiagaraTrackRecorderFactory::CreateTrackRec
 	return NewObject<UMovieSceneNiagaraTrackRecorder>();
 }
 
+UMovieSceneTrackRecorder* FMovieSceneNiagaraTrackRecorderFactory::CreateTrackRecorderForCacheTrack(IMovieSceneCachedTrack* CachedTrack, TObjectPtr<ULevelSequence> Sequence, TSharedPtr<ISequencer> Sequencer) const
+{
+	if (UMovieSceneNiagaraCacheTrack* NiagaraCacheTrack = Cast<UMovieSceneNiagaraCacheTrack>(CachedTrack))
+	{
+		UMovieSceneNiagaraTrackRecorder* TrackRecorder = NewObject<UMovieSceneNiagaraTrackRecorder>();
+		
+		UMovieScene* MovieScene = Sequence->GetMovieScene();
+		const TArray<FMovieSceneBinding>& SceneBindings = MovieScene->GetBindings();
+
+		for (const FMovieSceneBinding& Binding : SceneBindings)
+		{
+			TArray<UMovieSceneTrack*> ComponentTracks = Binding.GetTracks();
+			// find the Niagara component the track is bound to
+			if (ComponentTracks.Contains(NiagaraCacheTrack))
+			{
+				FGuid ObjectGuid = Binding.GetObjectGuid();
+				TArrayView<TWeakObjectPtr<>> BoundObjects = Sequencer->FindBoundObjects(ObjectGuid, Sequencer->GetFocusedTemplateID());
+				for (auto& Bound : BoundObjects)
+				{
+					if (UNiagaraComponent* NiagaraComponent = Cast<UNiagaraComponent>(Bound))
+					{
+						TrackRecorder->SystemToRecord = NiagaraComponent;
+						break;
+					}
+				}
+				TrackRecorder->NiagaraCacheTrack = NiagaraCacheTrack;
+				TrackRecorder->ObjectGuid = ObjectGuid;
+				TrackRecorder->OwningTakeRecorderSource = nullptr;
+				TrackRecorder->ObjectToRecord = TrackRecorder->SystemToRecord;
+				TrackRecorder->MovieScene = MovieScene;
+				TrackRecorder->Settings = nullptr;
+
+				for (UMovieSceneTrack* Track : ComponentTracks)
+				{
+					if (UMovieSceneNiagaraSystemTrack* SystemTrack = Cast<UMovieSceneNiagaraSystemTrack>(Track))
+					{
+						if (SystemTrack->IsEvalDisabled() == false)
+						{
+							for (UMovieSceneSection* Section : SystemTrack->GetAllSections())
+							{
+								if (!TrackRecorder->RecordRange.IsSet())
+								{
+									TrackRecorder->RecordRange = Section->GetRange();
+								}
+								TrackRecorder->RecordRange = FFrameNumberRange::Hull(*TrackRecorder->RecordRange, Section->GetRange());
+							}
+						}
+					}
+				}
+
+				TArray<UMovieSceneSection*> SceneSections = NiagaraCacheTrack->GetAllSections();
+				if (SceneSections.Num() > 0)
+				{
+					TrackRecorder->NiagaraCacheSection = Cast<UMovieSceneNiagaraCacheSection>(SceneSections[0]);
+				}
+				else
+				{
+					TrackRecorder->NiagaraCacheSection = CastChecked<UMovieSceneNiagaraCacheSection>(NiagaraCacheTrack->CreateNewSection());
+					TrackRecorder->NiagaraCacheSection->SetIsActive(false);
+					NiagaraCacheTrack->AddSection(*TrackRecorder->NiagaraCacheSection);
+				}
+				NiagaraCacheTrack->bIsRecording = true;
+
+				return TrackRecorder;
+			}
+		}
+	}
+	return nullptr;
+}
+
 UMovieSceneTrackRecorder* FMovieSceneNiagaraTrackRecorderFactory::CreateTrackRecorderForProperty(UObject* InObjectToRecord, const FName& InPropertyToRecord) const
 {
 	return nullptr;
@@ -45,6 +118,7 @@ void UMovieSceneNiagaraTrackRecorder::CreateTrackImpl()
 
 	if (NiagaraCacheTrack.IsValid())
 	{
+		NiagaraCacheTrack->bIsRecording = true;
 		NiagaraCacheSection = CastChecked<UMovieSceneNiagaraCacheSection>(NiagaraCacheTrack->CreateNewSection());
 		NiagaraCacheSection->SetIsActive(false);
 		NiagaraCacheTrack->AddSection(*NiagaraCacheSection);
@@ -57,31 +131,55 @@ void UMovieSceneNiagaraTrackRecorder::CreateTrackImpl()
 	}
 }
 
+bool UMovieSceneNiagaraTrackRecorder::ShouldContinueRecording(const FQualifiedFrameTime& FrameTime) const
+{
+	if (RecordRange.IsSet())
+	{
+		FFrameRate TickResolution = MovieScene->GetTickResolution();
+		FFrameNumber CurrentFrame = FrameTime.ConvertTo(TickResolution).FloorToFrame();
+		return CurrentFrame <= RecordRange.GetValue().GetUpperBoundValue();
+	}
+	return true;
+}
+
 void UMovieSceneNiagaraTrackRecorder::SetSectionStartTimecodeImpl(const FTimecode& InSectionStartTimecode, const FFrameNumber& InSectionFirstFrame)
 {
 	if (NiagaraCacheSection.IsValid() && NiagaraCacheTrack.IsValid())
 	{
-		NiagaraCacheSection->TimecodeSource = FMovieSceneTimecodeSource(InSectionStartTimecode);
+		const FFrameRate TickResolution = MovieScene->GetTickResolution();
+		const FFrameRate DisplayRate = MovieScene->GetDisplayRate();
 
 		FTakeRecorderParameters Parameters;
 		Parameters.User = GetDefault<UTakeRecorderUserSettings>()->Settings;
 		Parameters.Project = GetDefault<UTakeRecorderProjectSettings>()->Settings;
-
-		const FFrameRate TickResolution = MovieScene->GetTickResolution();
-		const FFrameRate DisplayRate = MovieScene->GetDisplayRate();
-
 		RecordStartTime = FApp::GetCurrentTime();
 		RecordStartFrame = Parameters.Project.bStartAtCurrentTimecode ?
 			FFrameRate::TransformTime(FFrameTime(InSectionStartTimecode.ToFrameNumber(DisplayRate)),
 				DisplayRate, TickResolution).FloorToFrame() : MovieScene->GetPlaybackRange().GetLowerBoundValue();
 
+		if (RecordRange.IsSet())
+		{
+			NiagaraCacheSection->TimecodeSource = FMovieSceneTimecodeSource(FTimecode::FromFrameNumber(RecordRange->GetLowerBoundValue(), TickResolution));
+			NiagaraCacheSection->SetRange(RecordRange.GetValue());
+			NiagaraCacheSection->SetStartFrame(TRangeBound<FFrameNumber>::Inclusive(RecordRange->GetLowerBoundValue()));
+		} 
+		else
+		{
+			NiagaraCacheSection->TimecodeSource = FMovieSceneTimecodeSource(InSectionStartTimecode);
+		}
+
 		if (SystemToRecord.IsValid())
 		{
 			// start simulation and writing to the sim cache 
 			SystemToRecord->SetSimCache(nullptr);
 			//SystemToRecord->SetAgeUpdateMode(ENiagaraAgeUpdateMode::TickDeltaTime);
 			//SystemToRecord->ResetSystem();
+			if (NiagaraCacheSection->Params.SimCache == nullptr)
+			{
+				NiagaraCacheSection->Params.SimCache = NewObject<UNiagaraSimCache>(NiagaraCacheSection.Get(), NAME_None, RF_Transactional);
+			}
 			NiagaraCacheSection->Params.SimCache->BeginWrite(NiagaraCacheSection->Params.CacheParameters, SystemToRecord.Get());
+			NiagaraCacheSection->bCacheOutOfDate = false;
 		}
 	}
 }
@@ -93,16 +191,27 @@ UMovieSceneSection* UMovieSceneNiagaraTrackRecorder::GetMovieSceneSection() cons
 
 void UMovieSceneNiagaraTrackRecorder::FinalizeTrackImpl()
 {
+	if (NiagaraCacheTrack.IsValid())
+	{
+		NiagaraCacheTrack->bIsRecording = false;
+	}
 	if (NiagaraCacheSection.IsValid())
 	{
 		// finalize the sim cache
 		NiagaraCacheSection->Params.SimCache->EndWrite();
 		
 		// Set the final range 
-		TOptional<TRange<FFrameNumber> > DefaultSectionLength = NiagaraCacheSection->GetAutoSizeRange();
+		TOptional<TRange<FFrameNumber>> DefaultSectionLength = NiagaraCacheSection->GetAutoSizeRange();
 		if (DefaultSectionLength.IsSet())
 		{
-			NiagaraCacheSection->SetRange(DefaultSectionLength.GetValue());
+			if (RecordRange.IsSet())
+			{
+				NiagaraCacheSection->SetRange(TRange<FFrameNumber>(RecordRange->GetLowerBoundValue(), (DefaultSectionLength->GetUpperBoundValue() - DefaultSectionLength->GetLowerBoundValue()) + RecordRange->GetLowerBoundValue()));
+			}
+			else
+			{
+				NiagaraCacheSection->SetRange(DefaultSectionLength.GetValue());
+			}
 		}
 
 		// Activate the section
@@ -112,14 +221,26 @@ void UMovieSceneNiagaraTrackRecorder::FinalizeTrackImpl()
 
 void UMovieSceneNiagaraTrackRecorder::RecordSampleImpl(const FQualifiedFrameTime& CurrentFrameTime)
 {
-	// Expand the section to the new length
-	const FFrameRate	 TickResolution = MovieScene->GetTickResolution();
+	const FFrameRate TickResolution = MovieScene->GetTickResolution();
 	const FFrameNumber CurrentFrame = CurrentFrameTime.ConvertTo(TickResolution).FloorToFrame();
 
+	if (RecordRange.IsSet() && (RecordRange->GetLowerBoundValue() > CurrentFrame || RecordRange->GetUpperBoundValue() < CurrentFrame))
+	{
+		return;
+	}
+
 	if (NiagaraCacheSection.IsValid() && SystemToRecord.IsValid())
 	{
-		NiagaraCacheSection->Params.SimCache->WriteFrame(SystemToRecord.Get());
-		NiagaraCacheSection->SetEndFrame(CurrentFrame);
+		FNiagaraSimCacheWriteResult Result = NiagaraCacheSection->Params.SimCache->WriteFrame(SystemToRecord.Get());
+		if (!Result.bSuccess)
+		{
+			UE_LOG(LogNiagaraSimCachingEditor, Warning, TEXT("Error trying to write sim cache: %s"), *Result.ErrorMsg);
+		}
+		else
+		{
+			// Expand the section to the new length
+			NiagaraCacheSection->SetEndFrame(CurrentFrame);
+		}
 	}
 }
 
diff --git a/Engine/Plugins/FX/NiagaraSimCaching/Source/NiagaraSimCachingEditor/Private/Sequencer/NiagaraCacheTrackEditor.cpp b/Engine/Plugins/FX/NiagaraSimCaching/Source/NiagaraSimCachingEditor/Private/Sequencer/NiagaraCacheTrackEditor.cpp
--- a/Engine/Plugins/FX/NiagaraSimCaching/Source/NiagaraSimCachingEditor/Private/Sequencer/NiagaraCacheTrackEditor.cpp
+++ b/Engine/Plugins/FX/NiagaraSimCaching/Source/NiagaraSimCachingEditor/Private/Sequencer/NiagaraCacheTrackEditor.cpp
@@ -1,19 +1,28 @@
 ﻿// Copyright Epic Games, Inc. All Rights Reserved.
 
 #include "Sequencer/NiagaraCacheTrackEditor.h"
-#include "CommonMovieSceneTools.h"
-#include "Fonts/FontMeasure.h"
-#include "Framework/Application/SlateApplication.h"
-#include "Framework/MultiBox/MultiBoxBuilder.h"
+
+#include "CoreMinimal.h"
+#include "AssetToolsModule.h"
 #include "LevelSequence.h"
 #include "NiagaraComponent.h"
 #include "NiagaraSimCache.h"
-#include "Niagara/Sequencer/MovieSceneNiagaraCacheTrack.h"
-#include "Niagara/Sequencer/MovieSceneNiagaraCacheSection.h"
+#include "NiagaraSystem.h"
 #include "SequencerSectionPainter.h"
-#include "SequencerUtilities.h"
-#include "Styling/SlateIconFinder.h"
 #include "TimeToPixel.h"
+#include "CacheTrackRecorder/Public/Recorder/CacheTrackRecorder.h"
+#include "Fonts/FontMeasure.h"
+#include "Framework/Application/SlateApplication.h"
+#include "Framework/MultiBox/MultiBoxBuilder.h"
+#include "Niagara/NiagaraSimCachingEditorStyle.h"
+#include "Niagara/Sequencer/MovieSceneNiagaraCacheSection.h"
+#include "Niagara/Sequencer/MovieSceneNiagaraCacheTrack.h"
+#include "Editor.h"
+#include "Subsystems/AssetEditorSubsystem.h"
+#include "Styling/SlateIconFinder.h"
+#include "UObject/Package.h"
+#include "Widgets/Images/SImage.h"
+#include "Widgets/Input/SButton.h"
 
 namespace NiagaraCacheEditorConstants
 {
@@ -82,7 +91,7 @@ int32 FNiagaraCacheSection::OnPaintSection(FSequencerSectionPainter& Painter) co
 	const float SeqLength = Duration - TickResolution.AsSeconds(Section.Params.StartFrameOffset + Section.Params.EndFrameOffset) / AnimPlayRate;
 	const float FirstLoopSeqLength = SeqLength - TickResolution.AsSeconds(Section.Params.FirstLoopStartFrameOffset) / AnimPlayRate;
 
-	if (!FMath::IsNearlyZero(SeqLength, KINDA_SMALL_NUMBER) && SeqLength > 0)
+	if (!FMath::IsNearlyZero(SeqLength, KINDA_SMALL_NUMBER) && SeqLength > 0 && Section.Params.SectionStretchMode == ENiagaraSimCacheSectionStretchMode::Repeat)
 	{
 		float MaxOffset = Section.GetRange().Size<FFrameTime>() / TickResolution;
 		float OffsetTime = FirstLoopSeqLength;
@@ -121,6 +130,12 @@ int32 FNiagaraCacheSection::OnPaintSection(FSequencerSectionPainter& Painter) co
 			// Draw the current time next to the scrub handle
 			const float AnimTime = Section.MapTimeToAnimation(Duration, CurrentTime, TickResolution);
 			int32 FrameTime = Duration ? FMath::FloorToFloat((AnimTime / Duration) * NumFrames) + 1 : 0;
+			if (Section.Params.SectionStretchMode == ENiagaraSimCacheSectionStretchMode::TimeDilate && Duration)
+			{
+				float Frames = (Section.GetExclusiveEndFrame() - Section.GetInclusiveStartFrame()).Value;
+				float NormalizedAge = (CurrentTime.FrameNumber - Section.GetInclusiveStartFrame()).Value / Frames;
+				FrameTime = FMath::FloorToFloat(NormalizedAge * NumFrames) + 1;
+			}
 			FString FrameString =  FString::FromInt(FrameTime);
 
 			const FSlateFontInfo SmallLayoutFont = FCoreStyle::GetDefaultFontStyle("Bold", 10);
@@ -141,7 +156,7 @@ int32 FNiagaraCacheSection::OnPaintSection(FSequencerSectionPainter& Painter) co
 
 			FSlateDrawElement::MakeBox(
 				Painter.DrawElements,
-				LayerId + 5,
+				LayerId + 100,
 				Painter.SectionGeometry.ToPaintGeometry(TextSize + 2.0f * BoxPadding, FSlateLayoutTransform(TextOffset - BoxPadding)),
 				FAppStyle::GetBrush("WhiteBrush"),
 				ESlateDrawEffect::None,
@@ -150,7 +165,7 @@ int32 FNiagaraCacheSection::OnPaintSection(FSequencerSectionPainter& Painter) co
 
 			FSlateDrawElement::MakeText(
 				Painter.DrawElements,
-				LayerId + 6,
+				LayerId + 101,
 				Painter.SectionGeometry.ToPaintGeometry(TextSize, FSlateLayoutTransform(TextOffset)),
 				FrameString,
 				SmallLayoutFont,
@@ -260,17 +275,14 @@ TSharedRef<ISequencerSection> FNiagaraCacheTrackEditor::MakeSectionInterface(UMo
 
 void FNiagaraCacheTrackEditor::BuildObjectBindingTrackMenu(FMenuBuilder& MenuBuilder, const TArray<FGuid>& ObjectBindings, const UClass* ObjectClass)
 {
-	//TODO (mga) allow users to create cache tracks directly?
-	/*
 	if (ObjectClass->IsChildOf(UNiagaraComponent::StaticClass()) || ObjectClass->IsChildOf(AActor::StaticClass()))
 	{
 		if(ObjectBindings.Num() > 0)
 		{
-			if (UNiagaraComponent* NiagaraComponent = AcquireNiagaraComponentFromObjectGuid(ObjectBindings[0], GetSequencer()))
+			if (AcquireNiagaraComponentFromObjectGuid(ObjectBindings[0], GetSequencer()))
 			{
 				UMovieSceneTrack* Track = nullptr;
 
-				
 				MenuBuilder.AddMenuEntry(
 					NSLOCTEXT("Sequencer", "AddNiagaraCache", "Niagara Cache"),
 					NSLOCTEXT("Sequencer", "AddNiagaraCacheTooltip", "Adds a Niagara Cache track."),
@@ -282,7 +294,62 @@ void FNiagaraCacheTrackEditor::BuildObjectBindingTrackMenu(FMenuBuilder& MenuBui
 				
 			}
 		}
-	}*/
+	}
+}
+
+void FNiagaraCacheTrackEditor::BuildTrackContextMenu(FMenuBuilder& MenuBuilder, UMovieSceneTrack* Track)
+{
+	MenuBuilder.BeginSection("CacheActions", LOCTEXT("CacheActions", "Cache Actions"));
+	{
+		if (UMovieSceneNiagaraCacheTrack* CacheTrack = Cast<UMovieSceneNiagaraCacheTrack>(Track))
+		{
+			MenuBuilder.AddMenuEntry(
+				LOCTEXT("CreateAssetFromThisCache", "Save Cache To Asset"),
+				LOCTEXT("CreateAssetFromThisCacheToolTip", "Create a simulation cache asset from this track's simulation data. This is especially useful if the cached data is very big or should be iterated on outside of sequencer.\nOnly enabled if there is valid cache data available and it's not already saved in an asset."),
+				FSlateIcon(),
+				FUIAction(
+					FExecuteAction::CreateLambda([CacheTrack]()
+					{
+						UMovieSceneNiagaraCacheSection* CacheSection = Cast<UMovieSceneNiagaraCacheSection>(CacheTrack->GetAllSections()[0]);
+						UNiagaraSimCache* SimCacheToSave = CacheSection->Params.SimCache;
+
+						const FString PackagePath = FPackageName::GetLongPackagePath(SimCacheToSave->GetOutermost()->GetName());
+						const FName CacheName = FName(CacheTrack->GetDisplayName().ToString());
+						
+						FAssetToolsModule& AssetToolsModule = FModuleManager::LoadModuleChecked<FAssetToolsModule>("AssetTools");
+
+						// Duplicate and Save the sim cache.
+						UNiagaraSimCache* CreatedAsset = Cast<UNiagaraSimCache>(AssetToolsModule.Get().DuplicateAssetWithDialogAndTitle(CacheName.ToString(), PackagePath, SimCacheToSave, LOCTEXT("CreateCacheAssetDialogTitle", "Save Sim Cache As")));
+						if (CreatedAsset != nullptr)
+						{
+							GEditor->GetEditorSubsystem<UAssetEditorSubsystem>()->OpenEditorForAsset(CreatedAsset);
+
+							// Replace existing cache with asset
+							FScopedTransaction ScopedTransaction(LOCTEXT("CreateAssetFromCache", "Create asset from cache"));
+							CacheSection->Modify();
+							CacheSection->Params.SimCache = CreatedAsset;
+						}
+					}),
+					FCanExecuteAction::CreateLambda(
+						[CacheTrack]()
+						{
+							TArray<UMovieSceneSection*> Sections = CacheTrack->GetAllSections();
+							if (Sections.Num() == 1)
+							{
+								if (UMovieSceneNiagaraCacheSection* CacheSection = Cast<UMovieSceneNiagaraCacheSection>(Sections[0]))
+								{
+									return CacheSection->Params.SimCache && !CacheSection->Params.SimCache->IsEmpty() && !CacheSection->Params.SimCache->IsAsset();
+								}
+							}
+							return false;
+						}
+					)
+				)
+			);
+		}
+	}
+
+	MenuBuilder.EndSection();
 }
 
 void FNiagaraCacheTrackEditor::BuildNiagaraCacheTrack(TArray<FGuid> ObjectBindings, UMovieSceneTrack* Track)
@@ -322,10 +389,16 @@ FKeyPropertyResult FNiagaraCacheTrackEditor::AddKeyInternal(FFrameNumber KeyTime
 		{
 			Track->Modify();
 
-			UMovieSceneSection* NewSection = Cast<UMovieSceneNiagaraCacheTrack>(Track)->AddNewAnimation(KeyTime, NiagaraComponent);
+			UMovieSceneNiagaraCacheTrack* CacheTrack = Cast<UMovieSceneNiagaraCacheTrack>(Track);
+			UMovieSceneSection* NewSection = CacheTrack->AddNewAnimation(KeyTime, NiagaraComponent);
 			KeyPropertyResult.bTrackModified = true;
 			KeyPropertyResult.SectionsCreated.Add(NewSection);
 
+			if (NiagaraComponent->GetAsset())
+			{
+				CacheTrack->SetDisplayName(FText::Format(LOCTEXT("NiagaraCacheTrack_Name", "{0} Sim Cache"), FText::FromName(NiagaraComponent->GetAsset()->GetFName())));
+			}
+
 			GetSequencer()->EmptySelection();
 			GetSequencer()->SelectSection(NewSection);
 			GetSequencer()->ThrobSectionSelection();
@@ -335,9 +408,124 @@ FKeyPropertyResult FNiagaraCacheTrackEditor::AddKeyInternal(FFrameNumber KeyTime
 	return KeyPropertyResult;
 }
 
+FReply FNiagaraCacheTrackEditor::RecordCacheTrack(IMovieSceneCachedTrack* Track)
+{
+	UCacheTrackRecorder::RecordCacheTrack(Track, GetSequencer());
+	return FReply::Handled();
+}
+
 TSharedPtr<SWidget> FNiagaraCacheTrackEditor::BuildOutlinerEditWidget(const FGuid& ObjectBinding, UMovieSceneTrack* Track, const FBuildEditWidgetParams& Params)
 {
-	return TSharedPtr<SWidget>();
+	TSharedPtr<ISequencer> SequencerPtr = GetSequencer();
+	FWeakObjectPtr BoundObject = SequencerPtr.IsValid() ? SequencerPtr->FindSpawnedObjectOrTemplate(ObjectBinding) : nullptr;
+	
+	return SNew(SHorizontalBox)
+		+ SHorizontalBox::Slot()
+		.Padding(2, 0, 2, 0)
+		.VAlign(VAlign_Center)
+		.AutoWidth()
+		[
+			SNew(SButton)
+			.ToolTipText(LOCTEXT("NiagaraCacheTrack_RecordCache", "Refresh the data in this cache by playing the sequence and recording the Niagara system."))
+			.HAlign(HAlign_Center)
+			.VAlign(VAlign_Center)
+			.ContentPadding(2)
+			.ForegroundColor(FSlateColor::UseForeground())
+			.ButtonStyle(FAppStyle::Get(), "HoverHintOnly")
+			.IsEnabled_Lambda([Track]()
+			{
+				return !Track->IsEvalDisabled();
+			})
+			.OnClicked(this, &FNiagaraCacheTrackEditor::RecordCacheTrack, Cast<IMovieSceneCachedTrack>(Track))
+			[
+				SNew(SImage)
+				.Image(FNiagaraSimCachingEditorStyle::Get().GetBrush( "Niagara.SimCaching.RecordIconSmall"))
+			]
+		]
+		+ SHorizontalBox::Slot()
+		.Padding(2, 0, 2, 0)
+		.VAlign(VAlign_Center)
+		[
+			SNew(SImage)
+			.ToolTipText(LOCTEXT("AddNiagaraCache_OutOfDate", "This cache track is out of date, as it's properties were changed after recording. Consider re-recording the cached data."))
+			.Image(FAppStyle::Get().GetBrush("Icons.WarningWithColor"))
+			.Visibility_Lambda([Track]()
+			{
+				if (UMovieSceneNiagaraCacheTrack* NiagaraTrack = Cast<UMovieSceneNiagaraCacheTrack>(Track))
+				{
+					for (UMovieSceneSection* Section : NiagaraTrack->GetAllSections())
+					{
+						if (UMovieSceneNiagaraCacheSection* NiagaraCacheSection = Cast<UMovieSceneNiagaraCacheSection>(Section))
+						{
+							if (NiagaraCacheSection->bCacheOutOfDate)
+							{
+								return EVisibility::Visible;
+							}
+						}
+					}
+				}
+				return EVisibility::Collapsed;
+			})
+		]
+		+ SHorizontalBox::Slot()
+		.Padding(2, 0, 2, 0)
+		.VAlign(VAlign_Center)
+		[
+			SNew(SImage)
+			.ToolTipText(LOCTEXT("AddNiagaraCache_InvalidCache", "The cache used by this track does not fit the Niagara component and will be ignored. Re-Record the sim cache or select a different one."))
+			.Image(FAppStyle::Get().GetBrush("Icons.ErrorWithColor"))
+			.Visibility_Lambda([Track, BoundObject]()
+			{
+				if (BoundObject.IsValid() && BoundObject.Get()->IsA(UNiagaraComponent::StaticClass()))
+				{
+					UNiagaraComponent* NiagaraComponent = Cast<UNiagaraComponent>(BoundObject.Get());
+					if (UMovieSceneNiagaraCacheTrack* NiagaraTrack = Cast<UMovieSceneNiagaraCacheTrack>(Track))
+					{
+						for (UMovieSceneSection* Section : NiagaraTrack->GetAllSections())
+						{
+							if (UMovieSceneNiagaraCacheSection* NiagaraCacheSection = Cast<UMovieSceneNiagaraCacheSection>(Section))
+							{
+								if (NiagaraCacheSection->Params.SimCache && !NiagaraCacheSection->Params.SimCache->IsEmpty() && !NiagaraCacheSection->Params.SimCache->CanRead(NiagaraComponent->GetAsset()))
+								{
+									return EVisibility::Visible;
+								}
+							}
+						}
+					}
+				}
+				return EVisibility::Collapsed;
+			})
+		]
+		+ SHorizontalBox::Slot()
+		.Padding(2, 0, 2, 0)
+		.VAlign(VAlign_Center)
+		[
+			SNew(SImage)
+			.ToolTipText(LOCTEXT("NiagaraCache_Active", "This track is currently providing cached data to the Niagara component, no simulation is running."))
+			.Image(FSlateIconFinder::FindIconForClass(UNiagaraSimCache::StaticClass()).GetIcon())
+			.ColorAndOpacity(FNiagaraSimCachingEditorStyle::Get().GetColor("Niagara.SimCaching.StatusIcon.Color"))
+			.Visibility_Lambda([Track, BoundObject]()
+			{
+				if (BoundObject.IsValid() && BoundObject.Get()->IsA(UNiagaraComponent::StaticClass()))
+				{
+					UNiagaraComponent* NiagaraComponent = Cast<UNiagaraComponent>(BoundObject.Get());
+					if (UMovieSceneNiagaraCacheTrack* NiagaraTrack = Cast<UMovieSceneNiagaraCacheTrack>(Track))
+					{
+						for (UMovieSceneSection* Section : NiagaraTrack->GetAllSections())
+						{
+							if (UMovieSceneNiagaraCacheSection* NiagaraCacheSection = Cast<UMovieSceneNiagaraCacheSection>(Section))
+							{
+								if (!NiagaraTrack->IsEvalDisabled() && NiagaraCacheSection->Params.SimCache && NiagaraCacheSection->Params.SimCache->CanRead(NiagaraComponent->GetAsset()) && NiagaraComponent->GetSimCache() == NiagaraCacheSection->Params.SimCache && NiagaraComponent->IsActive())
+								{
+									return EVisibility::Visible;
+								}
+							}
+						}
+					}
+				}
+				return EVisibility::Collapsed;
+			})
+		];
 }
 
 const FSlateBrush* FNiagaraCacheTrackEditor::GetIconBrush() const
diff --git a/Engine/Plugins/FX/NiagaraSimCaching/Source/NiagaraSimCachingEditor/Private/Sequencer/TakeRecorderNiagaraCacheSource.cpp b/Engine/Plugins/FX/NiagaraSimCaching/Source/NiagaraSimCachingEditor/Private/Sequencer/TakeRecorderNiagaraCacheSource.cpp
deleted file mode 100644
diff --git a/Engine/Plugins/FX/NiagaraSimCaching/Source/NiagaraSimCachingEditor/Public/Niagara/NiagaraSimCachingEditorPlugin.h b/Engine/Plugins/FX/NiagaraSimCaching/Source/NiagaraSimCachingEditor/Public/Niagara/NiagaraSimCachingEditorPlugin.h
--- a/Engine/Plugins/FX/NiagaraSimCaching/Source/NiagaraSimCachingEditor/Public/Niagara/NiagaraSimCachingEditorPlugin.h
+++ b/Engine/Plugins/FX/NiagaraSimCaching/Source/NiagaraSimCachingEditor/Public/Niagara/NiagaraSimCachingEditorPlugin.h
@@ -14,6 +14,8 @@ class UToolMenu;
 struct FToolMenuSection;
 class AActor;
 
+DECLARE_LOG_CATEGORY_EXTERN(LogNiagaraSimCachingEditor, Verbose, All);
+
 /**
  * The public interface to this module
  */
diff --git a/Engine/Plugins/FX/NiagaraSimCaching/Source/NiagaraSimCachingEditor/Public/Sequencer/MovieSceneNiagaraTrackRecorder.h b/Engine/Plugins/FX/NiagaraSimCaching/Source/NiagaraSimCachingEditor/Public/Sequencer/MovieSceneNiagaraTrackRecorder.h
--- a/Engine/Plugins/FX/NiagaraSimCaching/Source/NiagaraSimCachingEditor/Public/Sequencer/MovieSceneNiagaraTrackRecorder.h
+++ b/Engine/Plugins/FX/NiagaraSimCaching/Source/NiagaraSimCachingEditor/Public/Sequencer/MovieSceneNiagaraTrackRecorder.h
@@ -4,9 +4,7 @@
 
 #include "TrackRecorders/MovieSceneTrackRecorder.h"
 #include "CoreMinimal.h"
-#include "MovieScene.h"
 #include "TrackRecorders/IMovieSceneTrackRecorderFactory.h"
-#include "Channels/MovieSceneFloatChannel.h"
 
 #include "MovieSceneNiagaraTrackRecorder.generated.h"
 
@@ -26,6 +24,7 @@ public:
 	// ~Begin IMovieSceneTrackRecorderFactory Interface
 	virtual bool CanRecordObject(class UObject* InObjectToRecord) const override;
 	virtual UMovieSceneTrackRecorder* CreateTrackRecorderForObject() const override;
+	virtual class UMovieSceneTrackRecorder* CreateTrackRecorderForCacheTrack(IMovieSceneCachedTrack* CachedTrack, TObjectPtr<ULevelSequence> Sequence, TSharedPtr<ISequencer> Sequencer) const override;
 
 	// Particle Systems are entire components and you can't animate them as a property
 	virtual bool CanRecordProperty(class UObject* InObjectToRecord, class FProperty* InPropertyToRecord) const override { return false; }
@@ -35,8 +34,7 @@ public:
 };
 
 UCLASS(BlueprintType)
-class NIAGARASIMCACHINGEDITOR_API UMovieSceneNiagaraTrackRecorder
-	: public UMovieSceneTrackRecorder
+class NIAGARASIMCACHINGEDITOR_API UMovieSceneNiagaraTrackRecorder : public UMovieSceneTrackRecorder
 {
 	GENERATED_BODY()
 
@@ -47,6 +45,7 @@ public:
 	virtual void SetSectionStartTimecodeImpl(const FTimecode& InSectionStartTimecode, const FFrameNumber& InSectionFirstFrame) override;
 	virtual UMovieSceneSection* GetMovieSceneSection() const override;
 	virtual void CreateTrackImpl() override;
+	virtual bool ShouldContinueRecording(const FQualifiedFrameTime& FrameTime) const override;
 	// ~End UMovieSceneTrackRecorder Interface
 
 	/** Returns the Niagara cache track on which the cache manager will be recorded */
@@ -68,5 +67,9 @@ private:
 
 	/** The frame at the start of this recording section */
 	FFrameNumber RecordStartFrame;
+
+	TOptional<FFrameNumberRange> RecordRange;
+
+	friend FMovieSceneNiagaraTrackRecorderFactory;
 };
 
diff --git a/Engine/Plugins/FX/NiagaraSimCaching/Source/NiagaraSimCachingEditor/Public/Sequencer/NiagaraCacheTrackEditor.h b/Engine/Plugins/FX/NiagaraSimCaching/Source/NiagaraSimCachingEditor/Public/Sequencer/NiagaraCacheTrackEditor.h
--- a/Engine/Plugins/FX/NiagaraSimCaching/Source/NiagaraSimCachingEditor/Public/Sequencer/NiagaraCacheTrackEditor.h
+++ b/Engine/Plugins/FX/NiagaraSimCaching/Source/NiagaraSimCachingEditor/Public/Sequencer/NiagaraCacheTrackEditor.h
@@ -7,6 +7,7 @@
 #include "ISequencerTrackEditor.h"
 #include "MovieSceneTrackEditor.h"
 
+class IMovieSceneCachedTrack;
 /**
  * Editor for Niagara Cache tracks
  */
@@ -34,7 +35,8 @@ public:
 public:
 
 	// ~Begin ISequencerTrackEditor interface
-	virtual void BuildObjectBindingTrackMenu(class FMenuBuilder& MenuBuilder, const TArray<FGuid>& ObjectBindings, const UClass* ObjectClass) override;
+	virtual void BuildObjectBindingTrackMenu(FMenuBuilder& MenuBuilder, const TArray<FGuid>& ObjectBindings, const UClass* ObjectClass) override;
+	virtual void BuildTrackContextMenu(FMenuBuilder& MenuBuilder, UMovieSceneTrack* Track) override; 
 	virtual TSharedRef<ISequencerSection> MakeSectionInterface(UMovieSceneSection& SectionObject, UMovieSceneTrack& Track, FGuid ObjectBinding) override;
 	virtual bool SupportsSequence(UMovieSceneSequence* InSequence) const override;
 	virtual bool SupportsType(TSubclassOf<UMovieSceneTrack> Type) const override;
@@ -49,6 +51,8 @@ private:
 
 	/** Delegate for AnimatablePropertyChanged in AddKey */
 	FKeyPropertyResult AddKeyInternal(FFrameNumber KeyTime, class UNiagaraComponent* NiagaraComponent, UMovieSceneTrack* Track);
+
+	FReply RecordCacheTrack(IMovieSceneCachedTrack* Track);
 };
 
 /** Niagara cache sequencer section */
@@ -61,11 +65,6 @@ public:
 	/** Constructor. */
 	FNiagaraCacheSection( UMovieSceneSection& InSection, TWeakPtr<ISequencer> InSequencer);
 
-	/** Virtual destructor. */
-	virtual ~FNiagaraCacheSection() { }
-
-public:
-
 	// ~Begin ISequencerSection interface
 	virtual UMovieSceneSection* GetSectionObject() override;
 	virtual FText GetSectionTitle() const override;
diff --git a/Engine/Plugins/FX/NiagaraSimCaching/Source/NiagaraSimCachingEditor/Public/Sequencer/TakeRecorderNiagaraCacheSource.h b/Engine/Plugins/FX/NiagaraSimCaching/Source/NiagaraSimCachingEditor/Public/Sequencer/TakeRecorderNiagaraCacheSource.h
deleted file mode 100644
diff --git a/Engine/Plugins/VirtualProduction/Takes/Source/CacheTrackRecorder/CacheTrackRecorder.Build.cs b/Engine/Plugins/VirtualProduction/Takes/Source/CacheTrackRecorder/CacheTrackRecorder.Build.cs
new file mode 100644
--- /dev/null
+++ b/Engine/Plugins/VirtualProduction/Takes/Source/CacheTrackRecorder/CacheTrackRecorder.Build.cs
@@ -0,0 +1,36 @@
+// Copyright Epic Games, Inc. All Rights Reserved.
+
+using UnrealBuildTool;
+
+public class CacheTrackRecorder : ModuleRules
+{
+	public CacheTrackRecorder(ReadOnlyTargetRules Target) : base(Target)
+	{
+		PublicIncludePaths.Add(ModuleDirectory + "/Public");
+		
+		PrivateDependencyModuleNames.AddRange(
+			new[] {
+				"Core",
+				"CoreUObject",
+				"Engine",
+				"LevelEditor",
+				"LevelSequence",
+				"LevelSequenceEditor",
+				"MovieScene",
+				"TakesCore",
+				"TakeMovieScene",
+				"Settings",
+				"Slate",
+				"SlateCore",
+				"UnrealEd",
+			}
+		);
+
+		PublicDependencyModuleNames.AddRange(
+			new[] {
+                "TakeTrackRecorders",
+                "Sequencer",
+            }
+        );
+    }
+}
diff --git a/Engine/Plugins/VirtualProduction/Takes/Source/CacheTrackRecorder/Private/CacheTrackRecorderModule.cpp b/Engine/Plugins/VirtualProduction/Takes/Source/CacheTrackRecorder/Private/CacheTrackRecorderModule.cpp
new file mode 100644
--- /dev/null
+++ b/Engine/Plugins/VirtualProduction/Takes/Source/CacheTrackRecorder/Private/CacheTrackRecorderModule.cpp
@@ -0,0 +1,57 @@
+// Copyright Epic Games, Inc. All Rights Reserved.
+
+#include "CacheTrackRecorderModule.h"
+
+#include "Editor.h"
+#include "ToolMenus.h"
+#include "Features/IModularFeatures.h"
+#include "Framework/Application/SlateApplication.h"
+#include "Framework/MultiBox/MultiBoxBuilder.h"
+#include "Framework/MultiBox/MultiBoxExtender.h"
+#include "Misc/Attribute.h"
+#include "Modules/ModuleManager.h"
+#include "Recorder/CacheTrackRecorder.h"
+#include "Textures/SlateIcon.h"
+#include "Widgets/Docking/SDockTab.h"
+#include "Widgets/Input/SButton.h"
+
+#define LOCTEXT_NAMESPACE "CacheTrackRecorderModule"
+
+IMPLEMENT_MODULE(FCacheTrackRecorderModule, CacheTrackRecorder);
+
+FCacheTrackRecorderModule::FCacheTrackRecorderModule()
+{
+}
+
+void FCacheTrackRecorderModule::StartupModule()
+{
+#if WITH_EDITOR
+	if (GEditor)
+	{
+		GEditor->OnEditorClose().AddRaw(this, &FCacheTrackRecorderModule::OnEditorClose);
+	}
+#endif
+}
+
+void FCacheTrackRecorderModule::OnEditorClose()
+{
+	if (UCacheTrackRecorder* ActiveRecorder = UCacheTrackRecorder::GetActiveRecorder())
+	{
+		ActiveRecorder->Stop();
+	}
+}
+
+void FCacheTrackRecorderModule::ShutdownModule()
+{
+#if WITH_EDITOR
+	if (GEditor)
+	{
+		GEditor->OnEditorClose().RemoveAll(this);
+	}
+
+	FCoreDelegates::OnPostEngineInit.RemoveAll(this);
+#endif
+}
+
+
+#undef LOCTEXT_NAMESPACE
diff --git a/Engine/Plugins/VirtualProduction/Takes/Source/CacheTrackRecorder/Private/CacheTrackRecorderModule.h b/Engine/Plugins/VirtualProduction/Takes/Source/CacheTrackRecorder/Private/CacheTrackRecorderModule.h
new file mode 100644
--- /dev/null
+++ b/Engine/Plugins/VirtualProduction/Takes/Source/CacheTrackRecorder/Private/CacheTrackRecorderModule.h
@@ -0,0 +1,22 @@
+// Copyright Epic Games, Inc. All Rights Reserved.
+
+#pragma once
+
+#include "Templates/SharedPointer.h"
+#include "Delegates/IDelegateInstance.h"
+#include "UObject/GCObject.h"
+#include "UObject/WeakObjectPtrTemplates.h"
+#include "Modules/ModuleInterface.h"
+#include "Widgets/Input/SCheckBox.h"
+
+class FCacheTrackRecorderModule : public IModuleInterface
+{
+public:
+	FCacheTrackRecorderModule();
+	
+private:
+	virtual void StartupModule() override;
+	virtual void ShutdownModule() override;
+
+	void OnEditorClose();
+};
diff --git a/Engine/Plugins/VirtualProduction/Takes/Source/CacheTrackRecorder/Private/Recorder/CacheTrackRecorder.cpp b/Engine/Plugins/VirtualProduction/Takes/Source/CacheTrackRecorder/Private/Recorder/CacheTrackRecorder.cpp
new file mode 100644
--- /dev/null
+++ b/Engine/Plugins/VirtualProduction/Takes/Source/CacheTrackRecorder/Private/Recorder/CacheTrackRecorder.cpp
@@ -0,0 +1,1105 @@
+// Copyright Epic Games, Inc. All Rights Reserved.
+
+#include "Recorder/CacheTrackRecorder.h"
+
+#include "CoreGlobals.h"
+#include "Engine/Engine.h"
+#include "Engine/EngineTypes.h"
+#include "ISequencer.h"
+#include "LevelSequence.h"
+#include "Misc/CommandLine.h"
+#include "Misc/Parse.h"
+#include "TakeMetaData.h"
+#include "TakesUtils.h"
+#include "Tickable.h"
+#include "Stats/Stats.h"
+#include "TrackRecorders/IMovieSceneTrackRecorderFactory.h"
+#include "TrackRecorders/MovieSceneTrackRecorder.h"
+#include "Tracks/MovieSceneCachedTrack.h"
+#include "Features/IModularFeatures.h"
+
+// Engine includes
+#include "GameFramework/WorldSettings.h"
+
+// UnrealEd includes
+#include "Editor.h"
+
+#include "UObject/GCObjectScopeGuard.h"
+
+// Slate includes
+#include "Framework/Notifications/NotificationManager.h"
+#include "Widgets/Notifications/SNotificationList.h"
+#include "Widgets/Notifications/INotificationWidget.h"
+#include "Widgets/Text/STextBlock.h"
+#include "Widgets/SBoxPanel.h"
+#include "Widgets/Images/SThrobber.h"
+#include "Widgets/Input/SButton.h"
+#include "Widgets/Layout/SBox.h"
+
+// LevelEditor includes
+#include "IAssetViewport.h"
+#include "ILevelSequenceEditorToolkit.h"
+#include "LevelEditor.h"
+#include "Subsystems/AssetEditorSubsystem.h"
+
+#include UE_INLINE_GENERATED_CPP_BY_NAME(CacheTrackRecorder)
+
+#define LOCTEXT_NAMESPACE "CacheTrackRecorder"
+
+FCacheRecorderUserParameters::FCacheRecorderUserParameters()
+	: bMaximizeViewport(false)
+	, CountdownSeconds(0.f)
+	, EngineTimeDilation(1.f)
+	, bResetPlayhead(true)
+	, bStopAtPlaybackEnd(true)
+{
+}
+
+FCacheRecorderProjectParameters::FCacheRecorderProjectParameters()
+	: RecordingClockSource(EUpdateClockSource::RelativeTimecode)
+	, bStartAtCurrentTimecode(false)
+	, bRecordTimecode(false)
+	, bShowNotifications(true)
+{}
+
+class SCacheTrackRecorderNotification : public SCompoundWidget, public INotificationWidget
+{
+public:
+	SLATE_BEGIN_ARGS(SCacheTrackRecorderNotification){}
+	SLATE_END_ARGS()
+
+	void SetOwner(TSharedPtr<SNotificationItem> InOwningNotification)
+	{
+		WeakOwningNotification = InOwningNotification;
+	}
+
+	void Construct(const FArguments&, UCacheTrackRecorder* InCacheTrackRecorder)
+	{
+		WeakRecorder = InCacheTrackRecorder;
+		CacheTrackRecorderState = InCacheTrackRecorder->GetState();
+
+		UTakeMetaData* TakeMetaData = InCacheTrackRecorder->GetSequence()->FindMetaData<UTakeMetaData>();
+		check(TakeMetaData);
+
+		ChildSlot
+		[
+			SNew(SBox)
+			.Padding(FMargin(15.0f))
+			[
+				SNew(SVerticalBox)
+
+				+ SVerticalBox::Slot()
+				.Padding(FMargin(0,0,0,5.0f))
+				.HAlign(HAlign_Right)
+				.AutoHeight()
+				[
+					SNew(SHorizontalBox)
+
+					+ SHorizontalBox::Slot()
+					.VAlign(VAlign_Center)
+					[
+						SNew(STextBlock)
+						.Font(FCoreStyle::Get().GetFontStyle(TEXT("NotificationList.FontBold")))
+						.Text(FText::Format(LOCTEXT("RecordingTitleFormat", "Take {0} of slate {1}"), FText::AsNumber(TakeMetaData->GetTakeNumber()), FText::FromString(TakeMetaData->GetSlate())))
+					]
+
+					+ SHorizontalBox::Slot()
+					.AutoWidth()
+					.Padding(FMargin(15.f,0,0,0))
+					[
+						SAssignNew(Throbber, SThrobber)
+					]
+				]
+
+				+ SVerticalBox::Slot()
+				.AutoHeight()
+				.Padding(FMargin(0,0,0,5.0f))
+				.HAlign(HAlign_Right)
+				[
+					SAssignNew(TextBlock, STextBlock)
+					.Font(FCoreStyle::Get().GetFontStyle(TEXT("NotificationList.FontLight")))
+					.Text(GetDetailText())
+				]
+
+				+ SVerticalBox::Slot()
+				.AutoHeight()
+				.HAlign(HAlign_Right)
+				[
+					SNew(SHorizontalBox)
+
+					+ SHorizontalBox::Slot()
+					.AutoWidth()
+					.Padding(FMargin(5.0f,0,0,0))
+					.VAlign(VAlign_Center)
+					[
+						SAssignNew(Button, SButton)
+						.Text(LOCTEXT("StopButton", "Stop"))
+						.OnClicked(this, &SCacheTrackRecorderNotification::ButtonClicked)
+					]
+				]
+			]
+		];
+	}
+
+private:
+
+	virtual void Tick(const FGeometry& AllottedGeometry, const double InCurrentTime, const float InDeltaTime) override
+	{
+		bool bCloseNotification = false;
+		bool bCloseImmediately = false;
+		if (WeakRecorder.IsStale())
+		{
+			// Reset so we don't continually close the notification
+			bCloseImmediately = true;
+		}
+		else if (UCacheTrackRecorder* Recorder = WeakRecorder.Get())
+		{
+			ECacheTrackRecorderState NewCacheTrackRecorderState = Recorder->GetState();
+
+			if (NewCacheTrackRecorderState == ECacheTrackRecorderState::CountingDown || NewCacheTrackRecorderState == ECacheTrackRecorderState::Started)
+			{
+				// When counting down the text may change on tick
+				TextBlock->SetText(GetDetailText());
+			}
+
+			if (NewCacheTrackRecorderState != CacheTrackRecorderState)
+			{
+				TextBlock->SetText(GetDetailText());
+
+				if (NewCacheTrackRecorderState == ECacheTrackRecorderState::Stopped || NewCacheTrackRecorderState == ECacheTrackRecorderState::Cancelled)
+				{
+					Throbber->SetVisibility(EVisibility::Collapsed);
+					Button->SetVisibility(EVisibility::Collapsed);
+
+					bCloseNotification = true;
+				}
+			}
+
+			CacheTrackRecorderState = NewCacheTrackRecorderState;
+		}
+
+		TSharedPtr<SNotificationItem> Owner = WeakOwningNotification.Pin();
+		if ((bCloseNotification || bCloseImmediately) && Owner.IsValid())
+		{
+			if (bCloseImmediately)
+			{
+				Owner->SetFadeOutDuration(0.f);
+				Owner->SetExpireDuration(0.f);
+			}
+
+			Owner->ExpireAndFadeout();
+
+			// Remove our reference to the owner now that it's fading out
+			Owner = nullptr;
+		}
+	}
+
+	FText GetDetailText() const
+	{
+		if (UCacheTrackRecorder* Recorder = WeakRecorder.Get())
+		{
+			if (Recorder->GetState() == ECacheTrackRecorderState::CountingDown)
+			{
+				return FText::Format(LOCTEXT("CountdownText", "Recording in {0}s..."), FText::AsNumber(FMath::CeilToInt(Recorder->GetCountdownSeconds())));
+			}
+			if (Recorder->GetState() == ECacheTrackRecorderState::Stopped)
+			{
+				return LOCTEXT("CompleteText", "Recording Complete");
+			}
+			if (Recorder->GetState() == ECacheTrackRecorderState::Cancelled)
+			{
+				return LOCTEXT("CancelledText", "Recording Cancelled");
+			}
+
+			ULevelSequence* LevelSequence = Recorder->GetSequence();
+			if (UTakeMetaData* TakeMetaData = LevelSequence ? LevelSequence->FindMetaData<UTakeMetaData>() : nullptr)
+			{
+				FFrameRate FrameRate = TakeMetaData->GetFrameRate();
+				FTimespan RecordingDuration = FDateTime::UtcNow() - TakeMetaData->GetTimestamp();
+
+				FFrameNumber TotalFrames = FFrameNumber(static_cast<int32>(FrameRate.AsDecimal() * RecordingDuration.GetTotalSeconds()));
+
+				FTimecode Timecode = FTimecode::FromFrameNumber(TotalFrames, FrameRate);
+
+				return FText::Format(LOCTEXT("RecordingTimecodeText", "Recording...{0}"), FText::FromString(Timecode.ToString()));
+			}
+		}
+
+		return LOCTEXT("RecordingText", "Recording...");
+	}
+
+	virtual TSharedRef<SWidget> AsWidget() override
+	{
+		return AsShared();
+	}
+
+	// Unused
+	virtual void OnSetCompletionState(SNotificationItem::ECompletionState InState) override
+	{
+	}
+
+private:
+
+	FReply ButtonClicked()
+	{
+		if (UCacheTrackRecorder* Recorder = WeakRecorder.Get())
+		{
+			Recorder->Stop();
+		}
+		return FReply::Handled();
+	}
+
+private:
+	TSharedPtr<SWidget> Button, Throbber, Hyperlink;
+	TSharedPtr<STextBlock> TextBlock;
+
+	ECacheTrackRecorderState CacheTrackRecorderState = ECacheTrackRecorderState::CountingDown;
+	TWeakPtr<SNotificationItem> WeakOwningNotification;
+	TWeakObjectPtr<UCacheTrackRecorder> WeakRecorder;
+};
+
+
+class FTickableCacheTrackRecorder : public FTickableGameObject
+{
+public:
+
+	TWeakObjectPtr<UCacheTrackRecorder> WeakRecorder;
+
+	virtual TStatId GetStatId() const override
+	{
+		RETURN_QUICK_DECLARE_CYCLE_STAT(FTickableCacheTrackRecorder, STATGROUP_Tickables);
+	}
+
+	//Make sure it always ticks, otherwise we can miss recording, in particularly when time code is always increasing throughout the system.
+	virtual ETickableTickType GetTickableTickType() const override { return ETickableTickType::Always; }
+
+	virtual bool IsTickableInEditor() const override
+	{
+		return true;
+	}
+
+	virtual UWorld* GetTickableGameObjectWorld() const override
+	{
+			UCacheTrackRecorder* Recorder = WeakRecorder.Get();
+		return Recorder ? Recorder->GetWorld() : nullptr;
+	}
+
+	virtual void Tick(float DeltaTime) override
+	{
+		if (UCacheTrackRecorder* Recorder = WeakRecorder.Get())
+		{
+			Recorder->Tick(DeltaTime);
+		}
+	}
+};
+
+FTickableCacheTrackRecorder TickableCacheTrackRecorder;
+
+// Static members of UCacheTrackRecorder
+static TStrongObjectPtr<UCacheTrackRecorder>& GetCurrentRecorder()
+{
+	static TStrongObjectPtr<UCacheTrackRecorder> CurrentRecorder;
+	return CurrentRecorder;
+}
+FOnCacheTrackRecordingInitialized UCacheTrackRecorder::OnRecordingInitializedEvent;
+
+// Static functions for UCacheTrackRecorder
+UCacheTrackRecorder* UCacheTrackRecorder::GetActiveRecorder()
+{
+	return GetCurrentRecorder().Get();
+}
+
+FOnCacheTrackRecordingInitialized& UCacheTrackRecorder::OnRecordingInitialized()
+{
+	return OnRecordingInitializedEvent;
+}
+
+void UCacheTrackRecorder::RecordCacheTrack(IMovieSceneCachedTrack* Track, TSharedPtr<ISequencer> Sequencer, FCacheRecorderParameters Parameters)
+{	
+	ULevelSequence* LevelSequence = Sequencer ? Cast<ULevelSequence>(Sequencer->GetFocusedMovieSceneSequence()) : nullptr;	
+	if (LevelSequence && Track)
+	{
+		Parameters.StartFrame = LevelSequence->GetMovieScene()->GetPlaybackRange().GetLowerBoundValue();
+		Parameters.User.CountdownSeconds = 0;
+		Parameters.Project.bStartAtCurrentTimecode = false;
+
+		// If not resetting the playhead, store the current time as the start frame for recording. 
+		// This will ultimately be the start of the playback range and the recording will begin from that time.
+		if (!Parameters.User.bResetPlayhead)
+		{
+			Parameters.StartFrame = Sequencer->GetLocalTime().Time.FrameNumber;
+		}
+
+		TArray<IMovieSceneCachedTrack*> CacheTracks;
+		CacheTracks.Add(Track);
+
+		UCacheTrackRecorder* NewRecorder = NewObject<UCacheTrackRecorder>(GetTransientPackage(), NAME_None, RF_Transient);
+		UTakeMetaData* TakeMetaData = LevelSequence->FindOrAddMetaData<UTakeMetaData>();
+		if (TakeMetaData->GetSlate().IsEmpty())
+		{
+			TakeMetaData->SetSlate(LevelSequence->GetName());
+		}
+
+		FText ErrorText = LOCTEXT("UnknownError", "An unknown error occurred when trying to start recording");
+		if (!NewRecorder->Initialize(LevelSequence, CacheTracks, TakeMetaData, Parameters, &ErrorText))
+		{
+			if (ensure(!ErrorText.IsEmpty()))
+			{
+				FNotificationInfo Info(ErrorText);
+				Info.ExpireDuration = 5.0f;
+				FSlateNotificationManager::Get().AddNotification(Info)->SetCompletionState(SNotificationItem::CS_Fail);
+			}
+		}
+	}
+}
+
+bool UCacheTrackRecorder::SetActiveRecorder(UCacheTrackRecorder* NewActiveRecorder)
+{
+	if (GetCurrentRecorder().IsValid())
+	{
+		return false;
+	}
+
+	GetCurrentRecorder().Reset(NewActiveRecorder);
+	TickableCacheTrackRecorder.WeakRecorder = GetCurrentRecorder().Get();
+	OnRecordingInitializedEvent.Broadcast(NewActiveRecorder);
+	return true;
+}
+
+// Non-static api for UCacheTrackRecorder
+
+UCacheTrackRecorder::UCacheTrackRecorder(const FObjectInitializer& ObjInit)
+	: Super(ObjInit)
+{
+	CountdownSeconds = 0.f;
+	SequenceAsset = nullptr;
+}
+
+bool UCacheTrackRecorder::Initialize(ULevelSequence* LevelSequenceBase, const TArray<IMovieSceneCachedTrack*>& InCacheTracks, const UTakeMetaData* MetaData, const FCacheRecorderParameters& InParameters, FText* OutError)
+{
+	FGCObjectScopeGuard GCGuard(this);
+
+	if (GetActiveRecorder())
+	{
+		if (OutError)
+		{
+			*OutError = LOCTEXT("RecordingInProgressError", "A recording is currently in progress.");
+		}
+		return false;
+	}
+
+	if (MetaData->GetSlate().IsEmpty())
+	{
+		if (OutError)
+		{
+			*OutError = LOCTEXT("NoSlateSpecifiedError", "No slate specified.");
+		}
+		return false;
+	}
+
+	if (InCacheTracks.Num() == 0)
+	{
+		*OutError = LOCTEXT("RecordingModeError", "No valid tracks selected for cache recording.");
+		return false;
+	}
+	CacheTracks.Empty();
+	for (IMovieSceneCachedTrack* Track : InCacheTracks)
+	{
+		CacheTracks.AddDefaulted_GetRef().Track = Track;
+	}
+
+	OnRecordingPreInitializeEvent.Broadcast(this);
+
+	FCacheRecorderParameters FinalParameters = InParameters;
+	if (!InitializeSequencer(LevelSequenceBase, OutError))
+	{
+		return false;
+	}
+
+	// The SequenceAsset is either LevelSequenceBase or the currently focused sequence
+	SequenceAsset = Cast<ULevelSequence>(WeakSequencer.Pin()->GetFocusedMovieSceneSequence());
+	SequenceAsset->MarkPackageDirty();
+
+	// -----------------------------------------------------------
+	// Anything after this point assumes successful initialization
+	// -----------------------------------------------------------
+
+	AddToRoot();
+
+	Parameters = FinalParameters;
+	State      = ECacheTrackRecorderState::CountingDown;
+
+	// Perform any other parameter-configurable initialization. Must have a valid world at this point.
+	InitializeFromParameters();
+
+	// Figure out which world we're recording from
+	DiscoverSourceWorld();
+
+	// Open a recording notification
+	if (ShouldShowNotifications())
+	{
+		TSharedRef<SCacheTrackRecorderNotification> Content = SNew(SCacheTrackRecorderNotification, this);
+
+		FNotificationInfo Info(Content);
+		Info.bFireAndForget = false;
+		Info.ExpireDuration = 5.f;
+
+		TSharedPtr<SNotificationItem> PendingNotification = FSlateNotificationManager::Get().AddNotification(Info);
+		Content->SetOwner(PendingNotification);
+	}
+
+	ensure(SetActiveRecorder(this));
+
+	TSharedPtr<ISequencer> Sequencer = WeakSequencer.Pin();
+	if (Sequencer.IsValid())
+	{
+		// If a start frame was specified, adjust the playback range before rewinding to the beginning of the playback range
+		UMovieScene* MovieScene = SequenceAsset->GetMovieScene();
+		MovieScene->SetPlaybackRange(TRange<FFrameNumber>(Parameters.StartFrame, MovieScene->GetPlaybackRange().GetUpperBoundValue()));
+
+		// Always start the recording at the beginning of the playback range
+		Sequencer->SetLocalTime(MovieScene->GetPlaybackRange().GetLowerBoundValue());
+
+		// Center the view range around the current time about to be captured
+		FAnimatedRange Range = Sequencer->GetViewRange();
+		FTimecode CurrentTime = FApp::GetTimecode();
+		FFrameRate FrameRate = MovieScene->GetDisplayRate();
+		FFrameNumber ViewRangeStart = CurrentTime.ToFrameNumber(FrameRate);
+		double ViewRangeStartSeconds = Parameters.Project.bStartAtCurrentTimecode ? FrameRate.AsSeconds(ViewRangeStart) : MovieScene->GetPlaybackRange().GetLowerBoundValue() / MovieScene->GetTickResolution();
+		TRange<double> NewRange(ViewRangeStartSeconds - 0.5f, ViewRangeStartSeconds + (Range.GetUpperBoundValue() - Range.GetLowerBoundValue()) + 0.5f);
+		Sequencer->SetViewRange(NewRange, EViewRangeInterpolation::Immediate);
+		Sequencer->SetClampRange(TRange(Sequencer->GetViewRange()));
+	}
+
+	return true;
+}
+
+void UCacheTrackRecorder::DiscoverSourceWorld()
+{
+	UWorld* WorldToRecordIn = nullptr;
+
+	for (const FWorldContext& WorldContext : GEngine->GetWorldContexts())
+	{
+		if (WorldContext.WorldType == EWorldType::PIE || WorldContext.WorldType == EWorldType::Game )
+		{
+			WorldToRecordIn = WorldContext.World();
+			break;
+		}
+		else if (WorldContext.WorldType == EWorldType::Editor)
+		{
+			WorldToRecordIn = WorldContext.World();
+		}
+	}
+
+	check(WorldToRecordIn);
+	WeakWorld = WorldToRecordIn;
+
+	bool bPlayInGame = WorldToRecordIn->WorldType == EWorldType::PIE || WorldToRecordIn->WorldType == EWorldType::Game;
+	// If recording via PIE, be sure to stop recording cleanly when PIE ends
+	if ( bPlayInGame )
+	{
+		FEditorDelegates::EndPIE.AddUObject(this, &UCacheTrackRecorder::HandlePIE);
+	}
+	// If not recording via PIE, be sure to stop recording if PIE Starts
+	if ( !bPlayInGame )
+	{
+		FEditorDelegates::BeginPIE.AddUObject(this, &UCacheTrackRecorder::HandlePIE);//reuse same function
+	}
+}
+
+bool UCacheTrackRecorder::InitializeSequencer(ULevelSequence* LevelSequence, FText* OutError)
+{
+	if ( GEditor != nullptr )
+	{
+		// Open the sequence and set the sequencer ptr
+		GEditor->GetEditorSubsystem<UAssetEditorSubsystem>()->OpenEditorForAsset(LevelSequence);
+
+		IAssetEditorInstance*        AssetEditor         = GEditor->GetEditorSubsystem<UAssetEditorSubsystem>()->FindEditorForAsset(LevelSequence, false);
+		ILevelSequenceEditorToolkit* LevelSequenceEditor = static_cast<ILevelSequenceEditorToolkit*>(AssetEditor);
+
+		WeakSequencer = LevelSequenceEditor ? LevelSequenceEditor->GetSequencer() : nullptr;
+
+		if (!WeakSequencer.Pin().IsValid())
+		{
+			if (OutError)
+			{
+				*OutError = FText::Format(LOCTEXT("FailedToOpenSequencerError", "Failed to open Sequencer for asset '{0}."), FText::FromString(LevelSequence->GetPathName()));
+			}
+			return false;
+		}
+	}
+		
+	return true;
+}
+
+void UCacheTrackRecorder::InitializeFromParameters()
+{
+	// Initialize the countdown delay
+	CountdownSeconds = Parameters.User.CountdownSeconds;
+
+	// Set the end recording frame if enabled
+	StopRecordingFrame = Parameters.User.bStopAtPlaybackEnd ? SequenceAsset->GetMovieScene()->GetPlaybackRange().GetUpperBoundValue() : TOptional<FFrameNumber>();
+
+	// Apply immersive mode if the parameters demand it
+	if (Parameters.User.bMaximizeViewport)
+	{
+		TSharedPtr<IAssetViewport> ActiveLevelViewport = FModuleManager::Get().LoadModuleChecked<FLevelEditorModule>("LevelEditor").GetFirstActiveViewport();
+
+		// If it's already immersive we just leave it alone
+		if (ActiveLevelViewport.IsValid() && !ActiveLevelViewport->IsImmersive())
+		{
+			ActiveLevelViewport->MakeImmersive(true/*bWantImmersive*/, false/*bAllowAnimation*/);
+
+			// Restore it when we're done
+			auto RestoreImmersiveMode = [WeakViewport = TWeakPtr<IAssetViewport>(ActiveLevelViewport)]
+			{
+				if (TSharedPtr<IAssetViewport> CleaupViewport = WeakViewport.Pin())
+				{
+					CleaupViewport->MakeImmersive(false/*bWantImmersive*/, false/*bAllowAnimation*/);
+				}
+			};
+			OnStopCleanup.Add(RestoreImmersiveMode);
+		}
+	}
+}
+
+bool UCacheTrackRecorder::ShouldShowNotifications()
+{
+	// -CacheTrackRecorderISHEADLESS in the command line can force headless behavior and disable the notifications.
+	static const bool bCmdLineCacheTrackRecorderIsHeadless = FParse::Param(FCommandLine::Get(), TEXT("CacheTrackRecorderISHEADLESS"));
+
+	return Parameters.Project.bShowNotifications
+		&& !bCmdLineCacheTrackRecorderIsHeadless
+		&& !FApp::IsUnattended()
+		&& !GIsRunningUnattendedScript;
+}
+
+void UCacheTrackRecorder::BackupEditorTickState()
+{
+	SavedState.bBackedUp = true;
+	SavedState.bUseFixedTimeStep = FApp::UseFixedTimeStep();
+	SavedState.FixedDeltaTime = FApp::GetFixedDeltaTime();
+}
+
+void UCacheTrackRecorder::ModifyEditorTickState()
+{
+	ensure(SavedState.bBackedUp == false);
+	if (SequenceAsset)
+	{
+		if (UMovieScene* MovieScene = SequenceAsset->GetMovieScene())
+		{
+			BackupEditorTickState();
+			
+			FFrameRate FrameRate = MovieScene->GetDisplayRate();
+			double TickDelta = FrameRate.IsValid() ? FrameRate.AsInterval() : (1 / 30.0);
+			
+			// Force the engine into fixed timestep mode. There may be a global delay on the job that passes a fixed
+			// number of frames, so we want those frames to always pass the same amount of time for determinism. 
+			FApp::SetUseFixedTimeStep(true);
+			FApp::SetFixedDeltaTime(TickDelta);
+		}
+	}
+}
+
+void UCacheTrackRecorder::RestoreEditorTickState()
+{
+	if (SavedState.bBackedUp)
+	{
+		SavedState.bBackedUp = false;
+		FApp::SetUseFixedTimeStep(SavedState.bUseFixedTimeStep);
+		FApp::SetFixedDeltaTime(SavedState.FixedDeltaTime);
+	}
+}
+
+UWorld* UCacheTrackRecorder::GetWorld() const
+{
+	return WeakWorld.Get();
+}
+
+void UCacheTrackRecorder::Tick(float DeltaTime)
+{
+	if (State == ECacheTrackRecorderState::CountingDown)
+	{
+		NumberOfTicksAfterPre = 0;
+		CountdownSeconds = FMath::Max(0.f, CountdownSeconds - DeltaTime);
+		if (CountdownSeconds > 0.f)
+		{
+			return;
+		}
+		PreRecord();
+	}
+	else if (State == ECacheTrackRecorderState::PreRecord)
+	{
+		if (++NumberOfTicksAfterPre == 2) //seems we need 2 ticks to make sure things are settled
+		{
+			State = ECacheTrackRecorderState::TickingAfterPre;
+		}
+	}
+	else if (State == ECacheTrackRecorderState::TickingAfterPre)
+	{
+		NumberOfTicksAfterPre = 0;
+		Start();
+		InternalTick(0.0f);
+	}
+	else if (State == ECacheTrackRecorderState::Started)
+	{
+		InternalTick(DeltaTime);
+	}
+}
+
+FQualifiedFrameTime UCacheTrackRecorder::GetRecordTime() const
+{
+	FQualifiedFrameTime RecordTime;
+
+	if (TSharedPtr<ISequencer> Sequencer = WeakSequencer.Pin(); Sequencer.IsValid())
+	{
+		RecordTime = Sequencer->GetGlobalTime();
+	}
+	else if (SequenceAsset)
+	{
+		if (UMovieScene* MovieScene = SequenceAsset->GetMovieScene())
+		{
+			FFrameRate FrameRate = MovieScene->GetDisplayRate();
+			FFrameRate TickResolution = MovieScene->GetTickResolution();
+
+			FTimecode CurrentTimecode = FApp::GetTimecode();
+		
+			FFrameNumber CurrentFrame = FFrameRate::TransformTime(FFrameTime(CurrentTimecode.ToFrameNumber(FrameRate)), FrameRate, TickResolution).FloorToFrame();
+			FFrameNumber FrameAtStart = FFrameRate::TransformTime(FFrameTime(TimecodeAtStart.ToFrameNumber(FrameRate)), FrameRate, TickResolution).FloorToFrame();
+
+			if (Parameters.Project.bStartAtCurrentTimecode)
+			{
+				RecordTime = FQualifiedFrameTime(CurrentFrame, TickResolution);
+			}
+			else
+			{
+				RecordTime = FQualifiedFrameTime(CurrentFrame - FrameAtStart, TickResolution);
+			}
+		}
+	}
+
+	return RecordTime;
+}
+
+void UCacheTrackRecorder::InternalTick(float DeltaTime)
+{
+	UE::MovieScene::FScopedSignedObjectModifyDefer FlushOnTick(true);
+
+	TSharedPtr<ISequencer> Sequencer = WeakSequencer.Pin();
+	FQualifiedFrameTime RecordTime = GetRecordTime();
+
+	CurrentFrameTime = RecordTime.Time;
+	bool ShouldContinue = false;
+	for (FCachedTrackSource& CacheSource : CacheTracks)
+	{
+		if (CacheSource.Recorder && CacheSource.Recorder->ShouldContinueRecording(RecordTime) && DeltaTime > 0.0f)
+		{
+			ShouldContinue = true;
+			CacheSource.Recorder->RecordSample(RecordTime);
+			if (UMovieSceneTrack* SceneTrack = Cast<UMovieSceneTrack>(CacheSource.Track))
+			{
+				FFrameNumber EndFrame = RecordTime.Time.CeilToFrame();
+				for (UMovieSceneSection* SubSection : SceneTrack->GetAllSections())
+				{
+					// expand section range to this recording time
+					SubSection->ExpandToFrame(EndFrame);
+				}
+			}
+		}
+	}
+	if (ShouldContinue == false && DeltaTime > 0.0f)
+	{
+		StopRecordingFrame = CurrentFrameTime.FrameNumber;
+	}
+	
+	if (Sequencer.IsValid())
+	{
+		FAnimatedRange Range = Sequencer->GetViewRange();
+		if (UMovieScene* MovieScene = SequenceAsset->GetMovieScene())
+		{
+			FFrameRate FrameRate = MovieScene->GetTickResolution();
+			double CurrentTimeSeconds = FrameRate.AsSeconds(CurrentFrameTime) + 0.5f;
+			CurrentTimeSeconds = CurrentTimeSeconds > Range.GetUpperBoundValue() ? CurrentTimeSeconds : Range.GetUpperBoundValue();
+			TRange<double> NewRange(Range.GetLowerBoundValue(), CurrentTimeSeconds);
+			Sequencer->SetViewRange(NewRange, EViewRangeInterpolation::Immediate);
+			Sequencer->SetClampRange(TRange(Sequencer->GetViewRange()));
+		}
+	}
+
+	if (StopRecordingFrame.IsSet() && CurrentFrameTime.FrameNumber >= StopRecordingFrame.GetValue())
+	{
+		Stop();
+	}
+}
+
+void UCacheTrackRecorder::PreRecord()
+{
+	State = ECacheTrackRecorderState::PreRecord;
+
+	TSharedPtr<ISequencer> Sequencer = WeakSequencer.Pin();
+	for (FCachedTrackSource& CacheSource : CacheTracks)
+	{
+		static const FName MovieSceneSectionRecorderFactoryName("MovieSceneTrackRecorderFactory");
+		TArray<IMovieSceneTrackRecorderFactory*> ModularFactories = IModularFeatures::Get().GetModularFeatureImplementations<IMovieSceneTrackRecorderFactory>(MovieSceneSectionRecorderFactoryName);
+		for (IMovieSceneTrackRecorderFactory* Factory : ModularFactories)
+		{
+			if (UMovieSceneTrackRecorder* TrackRecorder = Factory->CreateTrackRecorderForCacheTrack(CacheSource.Track, SequenceAsset, Sequencer))
+			{
+				CacheSource.Recorder = TrackRecorder;
+				break;
+			}
+		}
+	}
+	
+	// Refresh sequencer in case the movie scene data has mutated (ie. existing object bindings removed because they will be recorded again)
+	if (Sequencer.IsValid())
+	{
+		Sequencer->RefreshTree();
+	}
+
+	// Apply engine Time Dilation after the countdown, otherwise the countdown will be dilated as well!
+	UWorld* RecordingWorld = GetWorld();
+	check(RecordingWorld);
+	if (AWorldSettings* WorldSettings = RecordingWorld->GetWorldSettings())
+	{
+		const float ExistingCinematicTimeDilation = WorldSettings->CinematicTimeDilation;
+
+		const bool bInvalidTimeDilation = Parameters.User.EngineTimeDilation == 0.f;
+
+		if (bInvalidTimeDilation)
+		{
+			UE_LOG(LogTakesCore, Warning, TEXT("Time dilation cannot be 0. Ignoring time dilation for this recording."));
+		}
+
+		if (Parameters.User.EngineTimeDilation != ExistingCinematicTimeDilation && !bInvalidTimeDilation)
+		{
+			WorldSettings->CinematicTimeDilation = Parameters.User.EngineTimeDilation;
+
+			// Restore it when we're done
+			auto RestoreTimeDilation = [ExistingCinematicTimeDilation, WeakWorldSettings = MakeWeakObjectPtr(WorldSettings)]
+			{
+				if (AWorldSettings* CleaupWorldSettings = WeakWorldSettings.Get())
+				{
+					CleaupWorldSettings->CinematicTimeDilation = ExistingCinematicTimeDilation;
+				}
+			};
+			OnStopCleanup.Add(RestoreTimeDilation);
+		}
+	}
+
+	if (Parameters.Project.bCacheTrackRecorderControlsClockTime)
+	{
+		ModifyEditorTickState();
+
+		auto RestoreEditorTick = [this]()
+		{
+			RestoreEditorTickState();
+		};
+		OnStopCleanup.Add(RestoreEditorTick);
+	}
+}
+
+void UCacheTrackRecorder::Start()
+{
+	FTimecode Timecode = FApp::GetTimecode();
+
+	State = ECacheTrackRecorderState::Started;
+
+	TSharedPtr<ISequencer> Sequencer = WeakSequencer.Pin();
+
+	CurrentFrameTime = FFrameTime(0);
+	TimecodeAtStart = Timecode;
+
+	// Discard any entity tokens we have so that restore state does not take effect when we delete any sections that recording will be replacing.
+	if (Sequencer.IsValid())
+	{
+		Sequencer->PreAnimatedState.DiscardEntityTokens();
+	}
+
+	if (UMovieScene* MovieScene = SequenceAsset->GetMovieScene())
+	{
+		CachedPlaybackRange = MovieScene->GetPlaybackRange();
+		CachedClockSource = MovieScene->GetClockSource();
+		MovieScene->SetClockSource(Parameters.Project.bCacheTrackRecorderControlsClockTime ? EUpdateClockSource::Tick : Parameters.Project.RecordingClockSource);
+		if (Sequencer.IsValid())
+		{
+			Sequencer->ResetTimeController();
+		}
+
+		FFrameRate FrameRate = MovieScene->GetDisplayRate();
+		FFrameRate TickResolution = MovieScene->GetTickResolution();
+		FFrameNumber PlaybackStartFrame = Parameters.Project.bStartAtCurrentTimecode ? FFrameRate::TransformTime(FFrameTime(Timecode.ToFrameNumber(FrameRate)), FrameRate, TickResolution).FloorToFrame() : MovieScene->GetPlaybackRange().GetLowerBoundValue();
+
+		// Transform all the sections to start the playback start frame
+		FFrameNumber DeltaFrame = PlaybackStartFrame - MovieScene->GetPlaybackRange().GetLowerBoundValue();
+		if (DeltaFrame != 0)
+		{
+			for (UMovieSceneSection* Section : MovieScene->GetAllSections())
+			{
+				Section->MoveSection(DeltaFrame);
+			}
+		}
+
+		if (StopRecordingFrame.IsSet())
+		{
+			StopRecordingFrame = StopRecordingFrame.GetValue() + DeltaFrame;
+		}
+
+		// Set infinite playback range when starting recording. Playback range will be clamped to the bounds of the sections at the completion of the recording
+		MovieScene->SetPlaybackRange(TRange<FFrameNumber>(PlaybackStartFrame, TNumericLimits<int32>::Max() - 1), false);
+		if (Sequencer.IsValid())
+		{
+			Sequencer->SetGlobalTime(PlaybackStartFrame);
+			Sequencer->SetPlaybackStatus(EMovieScenePlayerStatus::Playing);
+		}
+	}
+
+	FQualifiedFrameTime RecordTime = GetRecordTime();
+	FFrameNumber RecordStartFrame = RecordTime.Time.FloorToFrame();
+	FFrameRate TargetLevelSequenceTickResolution = SequenceAsset->MovieScene->GetTickResolution();
+	FFrameRate TargetLevelSequenceDisplayRate = SequenceAsset->MovieScene->GetDisplayRate();
+	FTimecode CurrentTimecode = FTimecode::FromFrameNumber(FFrameRate::TransformTime(RecordTime.Time, TargetLevelSequenceTickResolution, TargetLevelSequenceDisplayRate).FloorToFrame(), TargetLevelSequenceDisplayRate);
+	for (FCachedTrackSource& CacheSource : CacheTracks)
+	{
+		if (CacheSource.Recorder)
+		{
+			CacheSource.Recorder->SetSectionStartTimecode(CurrentTimecode, RecordTime.Time.FloorToFrame());
+
+			// Record immediately so that there's a key on the first frame of recording
+			CacheSource.Recorder->RecordSample(RecordTime);
+			if (UMovieSceneTrack* SceneTrack = Cast<UMovieSceneTrack>(CacheSource.Track))
+			{
+				for (UMovieSceneSection* SubSection : SceneTrack->GetAllSections())
+				{
+					SubSection->TimecodeSource = FMovieSceneTimecodeSource(Timecode);
+
+					// Ensure we're expanded to at least the next frame so that we don't set the start past the end
+					// when we set the first frame.
+					SubSection->ExpandToFrame(RecordStartFrame + FFrameNumber(1));
+					SubSection->SetStartFrame(TRangeBound<FFrameNumber>::Inclusive(RecordStartFrame));
+				}
+			}
+		}
+	}
+	
+	if (!ShouldShowNotifications())
+	{
+		// Log in lieu of the notification widget
+		UE_LOG(LogTakesCore, Log, TEXT("Started recording"));
+	}
+
+	OnRecordingStartedEvent.Broadcast(this);
+}
+
+void UCacheTrackRecorder::Stop()
+{
+	constexpr bool bCancelled = false;
+	StopInternal(bCancelled);
+}
+
+void UCacheTrackRecorder::Cancel()
+{
+	constexpr bool bCancelled = true;
+	StopInternal(bCancelled);
+}
+
+void UCacheTrackRecorder::StopInternal(const bool bCancelled)
+{
+	static bool bStoppedRecording = false;
+
+	if (bStoppedRecording)
+	{
+		return;
+	}
+
+	double StartTime = FPlatformTime::Seconds();
+
+	TGuardValue<bool> ReentrantGuard(bStoppedRecording, true);
+
+	FEditorDelegates::EndPIE.RemoveAll(this);
+	FEditorDelegates::BeginPIE.RemoveAll(this);
+	
+	const bool bDidEverStartRecording = State == ECacheTrackRecorderState::Started;
+	const bool bRecordingFinished = !bCancelled && bDidEverStartRecording;
+	State = bRecordingFinished ? ECacheTrackRecorderState::Stopped : ECacheTrackRecorderState::Cancelled;
+
+	TSharedPtr<ISequencer> Sequencer = WeakSequencer.Pin();
+	if (Sequencer.IsValid())
+	{
+		Sequencer->SetPlaybackStatus(EMovieScenePlayerStatus::Stopped);
+	}
+
+	UMovieScene* MovieScene = SequenceAsset->GetMovieScene();
+	if (MovieScene)
+	{
+		MovieScene->SetClockSource(CachedClockSource);
+
+		if (Sequencer.IsValid())
+		{
+			Sequencer->ResetTimeController();
+		}
+	}
+
+	if (bDidEverStartRecording)
+	{
+		if (!ShouldShowNotifications())
+		{
+			// Log in lieu of the notification widget
+			if (bRecordingFinished)
+			{
+				UE_LOG(LogTakesCore, Log, TEXT("Stopped recording"));
+			}
+			else
+			{
+				UE_LOG(LogTakesCore, Log, TEXT("Recording cancelled"));
+			}
+		}
+
+		OnRecordingStoppedEvent.Broadcast(this);
+
+		if (MovieScene)
+		{
+			TRange<FFrameNumber> Range = MovieScene->GetPlaybackRange();
+
+			//Set Range to what we recorded instead of that large number, this let's us reliably set camera cut times.
+			FFrameNumber PlaybackEndFrame = CurrentFrameTime.FrameNumber;
+
+			if (StopRecordingFrame.IsSet())
+			{
+				PlaybackEndFrame = StopRecordingFrame.GetValue();
+			}
+			
+			MovieScene->SetPlaybackRange(TRange<FFrameNumber>(Range.GetLowerBoundValue(), PlaybackEndFrame));
+
+			if (Sequencer)
+			{
+				Sequencer->ResetTimeController();
+			}
+		}
+
+		// stop all recordings
+		for (FCachedTrackSource& CacheSource : CacheTracks)
+		{
+			if (CacheSource.Recorder)
+			{
+				CacheSource.Recorder->StopRecording();
+			}
+		}
+
+		// finalize tracks
+		for (FCachedTrackSource& CacheSource : CacheTracks)
+		{
+			if (CacheSource.Recorder)
+			{
+				CacheSource.Recorder->FinalizeTrack();
+			}
+		}
+
+		// Restore the playback range to what it was before recording.
+		if (MovieScene)
+		{
+			MovieScene->SetPlaybackRange(CachedPlaybackRange);
+		}
+
+		if (bRecordingFinished)
+		{
+			UTakeMetaData* AssetMetaData = SequenceAsset->FindMetaData<UTakeMetaData>();
+
+			if (GEditor && GEditor->GetEditorWorldContext().World() && ensure(AssetMetaData))
+			{
+				AssetMetaData->SetLevelOrigin(GEditor->GetEditorWorldContext().World()->PersistentLevel);
+			}
+
+			// Rebuild sequencer because subsequences could have been added or bindings removed
+			if (Sequencer)
+			{
+				Sequencer->RefreshTree();
+			}
+		}
+	}
+
+	// Perform any other cleanup that has been defined for this recording
+	for (const TFunction<void()>& Cleanup : OnStopCleanup)
+	{
+		Cleanup();
+	}
+	OnStopCleanup.Reset();
+
+	// reset the current recorder and stop us from being ticked
+	if (GetCurrentRecorder().Get() == this)
+	{
+		GetCurrentRecorder().Reset();
+		TickableCacheTrackRecorder.WeakRecorder = nullptr;
+
+		if (bRecordingFinished)
+		{
+			OnRecordingFinishedEvent.Broadcast(this);
+		}
+		else
+		{
+			OnRecordingCancelledEvent.Broadcast(this);
+		}
+	}
+
+	if (SequenceAsset)
+	{
+		UPackage* const Package = SequenceAsset->GetOutermost();
+		FString const PackageName = Package->GetName();
+
+		double ElapsedTime = FPlatformTime::Seconds() - StartTime;
+		UE_LOG(LogTakesCore, Log, TEXT("Finished processing %s in %0.2f seconds"), *PackageName, ElapsedTime);
+	}
+
+	RemoveFromRoot();
+}
+
+FOnCacheTrackRecordingPreInitialize& UCacheTrackRecorder::OnRecordingPreInitialize()
+{
+	return OnRecordingPreInitializeEvent;
+}
+
+FOnCacheTrackRecordingStarted& UCacheTrackRecorder::OnRecordingStarted()
+{
+	return OnRecordingStartedEvent;
+}
+
+FOnCacheTrackRecordingStopped& UCacheTrackRecorder::OnRecordingStopped()
+{
+	return OnRecordingStoppedEvent;
+}
+
+FOnCacheTrackRecordingFinished& UCacheTrackRecorder::OnRecordingFinished()
+{
+	return OnRecordingFinishedEvent;
+}
+
+FOnCacheTrackRecordingCancelled& UCacheTrackRecorder::OnRecordingCancelled()
+{
+	return OnRecordingCancelledEvent;
+}
+
+void UCacheTrackRecorder::HandlePIE(bool bIsSimulating)
+{
+	ULevelSequence* FinishedAsset = GetSequence();
+
+	if (ShouldShowNotifications())
+	{
+		TSharedRef<SCacheTrackRecorderNotification> Content = SNew(SCacheTrackRecorderNotification, this);
+
+		FNotificationInfo Info(Content);
+		Info.ExpireDuration = 5.f;
+
+		TSharedPtr<SNotificationItem> PendingNotification = FSlateNotificationManager::Get().AddNotification(Info);
+		PendingNotification->SetCompletionState(SNotificationItem::CS_Success);
+	}
+	
+	Stop();
+
+	if (FinishedAsset)
+	{
+		GEditor->GetEditorSubsystem<UAssetEditorSubsystem>()->CloseAllEditorsForAsset(FinishedAsset);
+	}
+}
+
+#undef LOCTEXT_NAMESPACE
+
diff --git a/Engine/Plugins/VirtualProduction/Takes/Source/CacheTrackRecorder/Public/Recorder/CacheTrackRecorder.h b/Engine/Plugins/VirtualProduction/Takes/Source/CacheTrackRecorder/Public/Recorder/CacheTrackRecorder.h
new file mode 100644
--- /dev/null
+++ b/Engine/Plugins/VirtualProduction/Takes/Source/CacheTrackRecorder/Public/Recorder/CacheTrackRecorder.h
@@ -0,0 +1,373 @@
+// Copyright Epic Games, Inc. All Rights Reserved.
+
+#pragma once
+
+#include "UObject/Object.h"
+#include "ISequencer.h"
+#include "CacheTrackRecorder.generated.h"
+
+class ULevelSequence;
+class UMovieSceneTrackRecorder;
+class IMovieSceneCachedTrack;
+class UTakeMetaData;
+
+USTRUCT()
+struct CACHETRACKRECORDER_API FCacheRecorderUserParameters
+{
+	GENERATED_BODY()
+
+	FCacheRecorderUserParameters();
+
+	/** Whether to maximize the viewport (enter Immersive Mode) when recording */
+	UPROPERTY(config, EditAnywhere, Category="User Settings")
+	bool bMaximizeViewport;
+
+	/** Delay that we will use before starting recording */
+	UPROPERTY(config, EditAnywhere, Category="User Settings", DisplayName="Countdown", meta=(Units=s, ClampMin="0.0", UIMin="0.0", ClampMax="60.0", UIMax="60.0"))
+	float CountdownSeconds;
+
+	/** The engine time dilation to apply during the recording */
+	UPROPERTY(config, EditAnywhere, Category="User Settings", meta=(Units=Multiplier, ClampMin="0.00001", UIMin="0.00001"))
+	float EngineTimeDilation;
+
+	/** Reset playhead to beginning of the playback range when starting recording */
+	UPROPERTY(config, EditAnywhere, Category = "User Settings")
+	bool bResetPlayhead;
+
+	/** Automatically stop recording when reaching the end of the playback range */
+	UPROPERTY(config, EditAnywhere, Category = "User Settings")
+	bool bStopAtPlaybackEnd;
+};
+
+USTRUCT()
+struct CACHETRACKRECORDER_API FCacheRecorderProjectParameters
+{
+	GENERATED_BODY()
+
+	FCacheRecorderProjectParameters();
+	
+	/**
+	 * The default name to use for the Slate information
+	 */
+	UPROPERTY(config, EditAnywhere, Category = "Cache Recorder")
+	FString DefaultSlate;
+
+	/**
+	 * If true then take recorder will control the sequencer timing when recording with a fixed editor time step. The delta time is derived by the sequence's target frame rate. 
+	 * This is useful when recording cache data where frame accuracy is important (e.g. Niagara systems), but should be set to false when dealing with data from external sources (e.g. LiveLink).
+	 */
+	UPROPERTY(config, EditAnywhere, Category = "Cache Recorder")
+	bool bCacheTrackRecorderControlsClockTime = true;
+
+	/**
+	 * The clock source to use when recording
+	 */
+	UPROPERTY(config, EditAnywhere, Category = "Cache Recorder", meta=(EditCondition="!bCacheTrackRecorderControlsClockTime"))
+	EUpdateClockSource RecordingClockSource;
+
+	/**
+	 * If enabled, track sections will start at the current timecode. Otherwise, 0.
+	 */
+	UPROPERTY(config, EditAnywhere, Category = "Cache Recorder")
+	bool bStartAtCurrentTimecode;
+
+	/**
+	 * If enabled, timecode will be recorded into each actor track
+	 */
+	UPROPERTY(config, EditAnywhere, Category = "Cache Recorder")
+	bool bRecordTimecode;
+
+	/** Whether to show notification windows or not when recording */
+	UPROPERTY(config, EditAnywhere, Category = "Cache Recorder")
+	bool bShowNotifications;
+};
+
+/**
+ * Structure housing all configurable parameters for a take recorder instance
+ */
+USTRUCT()
+struct CACHETRACKRECORDER_API FCacheRecorderParameters
+{
+	GENERATED_BODY()
+
+	UPROPERTY(EditAnywhere, Category="Cache Recorder")
+	FCacheRecorderUserParameters User;
+
+	UPROPERTY(EditAnywhere, Category="Cache Recorder")
+	FCacheRecorderProjectParameters Project;
+
+	UPROPERTY(EditAnywhere, Category = "Cache Recorder")
+	FFrameNumber StartFrame;
+};
+
+UENUM(BlueprintType)
+enum class ECacheTrackRecorderState : uint8
+{
+	CountingDown,
+	PreRecord,
+	TickingAfterPre,
+	Started,
+	Stopped,
+	Cancelled,
+};
+
+DECLARE_MULTICAST_DELEGATE_OneParam(FOnCacheTrackRecordingPreInitialize, UCacheTrackRecorder*);
+DECLARE_MULTICAST_DELEGATE_OneParam(FOnCacheTrackRecordingInitialized, UCacheTrackRecorder*);
+DECLARE_MULTICAST_DELEGATE_OneParam(FOnCacheTrackRecordingStarted, UCacheTrackRecorder*);
+DECLARE_MULTICAST_DELEGATE_OneParam(FOnCacheTrackRecordingStopped, UCacheTrackRecorder*);
+DECLARE_MULTICAST_DELEGATE_OneParam(FOnCacheTrackRecordingFinished, UCacheTrackRecorder*);
+DECLARE_MULTICAST_DELEGATE_OneParam(FOnCacheTrackRecordingCancelled, UCacheTrackRecorder*);
+
+USTRUCT()
+struct FCachedTrackSource
+{
+	GENERATED_USTRUCT_BODY()
+	
+	IMovieSceneCachedTrack* Track = nullptr;
+
+	UPROPERTY()
+	TObjectPtr<UMovieSceneTrackRecorder> Recorder;
+};
+
+UCLASS(BlueprintType)
+class CACHETRACKRECORDER_API UCacheTrackRecorder : public UObject
+{
+public:
+
+	GENERATED_BODY()
+
+	UCacheTrackRecorder(const FObjectInitializer& ObjInit);
+
+public:
+
+	/**
+	 * Retrieve the currently active take recorder instance
+	 */
+	static UCacheTrackRecorder* GetActiveRecorder();
+
+	/**
+	 * Retrieve a multi-cast delegate that is triggered when a new recording begins
+	 */
+	static FOnCacheTrackRecordingInitialized& OnRecordingInitialized();
+
+	/**
+	 * Utility method to quickly record a single cache track
+	 */
+	static void RecordCacheTrack(IMovieSceneCachedTrack* Track, TSharedPtr<ISequencer> Sequencer, FCacheRecorderParameters Parameters = FCacheRecorderParameters());
+public:
+
+	/**
+	 * Access the number of seconds remaining before this recording will start
+	 */
+	UFUNCTION(BlueprintCallable, Category="Cache Recorder")
+	float GetCountdownSeconds() const
+	{
+		return CountdownSeconds;
+	}
+
+	/**
+	 * Access the sequence asset that this recorder is recording into
+	 */
+	UFUNCTION(BlueprintCallable, Category="Cache Recorder")
+	ULevelSequence* GetSequence() const
+	{
+		return SequenceAsset;
+	}
+
+	/**
+	 * Get the current state of this recorder
+	 */
+	UFUNCTION(BlueprintCallable, Category="Cache Recorder")
+	ECacheTrackRecorderState GetState() const
+	{
+		return State;
+	}
+
+	/**
+	 * Initialize a new recording with the specified parameters. Fails if another recording is currently in progress.
+	 *
+	 */
+	bool Initialize(ULevelSequence* LevelSequenceBase, const TArray<IMovieSceneCachedTrack*>& CacheTracks, const UTakeMetaData* MetaData, const FCacheRecorderParameters& InParameters, FText* OutError = nullptr);
+
+	/**
+	 * Called to stop the recording
+	 */
+	void Stop();
+
+	/**
+	 * Called to cancel the recording
+	 */
+	void Cancel();
+
+	/**
+	 * Retrieve a multi-cast delegate that is triggered before initialization occurs (ie. when the recording button is pressed and before the countdown starts)
+	 */
+	FOnCacheTrackRecordingPreInitialize& OnRecordingPreInitialize();
+
+	/**
+	 * Retrieve a multi-cast delegate that is triggered when this recording starts
+	 */
+	FOnCacheTrackRecordingStarted& OnRecordingStarted();
+
+	/**
+	 * Retrieve a multi-cast delegate that is triggered when this recording is stopped
+	 */
+	FOnCacheTrackRecordingStopped& OnRecordingStopped();
+
+	/**
+	 * Retrieve a multi-cast delegate that is triggered when this recording finishes
+	 */
+	FOnCacheTrackRecordingFinished& OnRecordingFinished();
+
+	/**
+	 * Retrieve a multi-cast delegate that is triggered when this recording is cancelled
+	 */
+	FOnCacheTrackRecordingCancelled& OnRecordingCancelled();
+
+private:
+
+	/**
+	 * Called after the countdown to PreRecord
+	 */
+	void PreRecord();
+
+	/**
+	 * Called after PreRecord To Start
+	 */
+	void Start();
+
+	/*
+	 * Stop or cancel
+	 */
+	void StopInternal(const bool bCancelled);
+
+	/**
+	 * Ticked by a tickable game object to performe any necessary time-sliced logic
+	 */
+	void Tick(float DeltaTime);
+
+	/**
+	 * Called if we're currently recording a PIE world that has been shut down or if we start PIE in a non-PIE world. Bound in Initialize, and unbound in Stop.
+	 */
+	void HandlePIE(bool bIsSimulating);
+
+	/**
+	 * Attempt to open the sequencer UI for the asset to be recorded
+	 */
+	bool InitializeSequencer(ULevelSequence* LevelSequence, FText* OutError);
+
+	/**
+	 * Discovers the source world to record from, and initializes it for recording
+	 */
+	void DiscoverSourceWorld();
+
+	/**
+	 * Called to perform any initialization based on the user-provided parameters
+	 */
+	void InitializeFromParameters();
+
+	/**
+     * Returns true if notification widgets should be shown when recording.
+	 * It takes into account CacheTrackRecorder project settings, the command line, and global unattended settings.
+     */
+	bool ShouldShowNotifications();
+
+	/* Called before setting the editor into a fixed time step mode */
+	void BackupEditorTickState();
+
+	/* Sets the editor up to tick in a fixed time step */
+	void ModifyEditorTickState();
+
+	/* Restores the editor tick state if necessary */
+	void RestoreEditorTickState();
+
+	/* The time at which to record. Taken from the Sequencer global time, otherwise based on timecode */
+	FQualifiedFrameTime GetRecordTime() const;
+
+	/** Called by Tick and Start to make sure we record at start */
+	void InternalTick(float DeltaTime);
+
+	virtual UWorld* GetWorld() const override;
+
+private:
+
+	/** The number of seconds remaining before Start() should be called */
+	float CountdownSeconds;
+
+	/** The state of this recorder instance */
+	ECacheTrackRecorderState State;
+
+	/** FFrameTime in MovieScene Resolution we are at*/
+	FFrameTime CurrentFrameTime;
+
+	/** Timecode at the start of recording */
+	FTimecode TimecodeAtStart;
+
+	/** Optional frame to stop recording at*/
+	TOptional<FFrameNumber> StopRecordingFrame;
+
+	/** The asset that we should output recorded data into */
+	UPROPERTY(transient)
+	TObjectPtr<ULevelSequence> SequenceAsset;
+
+	/** The world that we are recording within */
+	UPROPERTY(transient)
+	TWeakObjectPtr<UWorld> WeakWorld;
+
+	/** Parameters for the recorder - marked up as a uproperty to support reference collection */
+	UPROPERTY()
+	FCacheRecorderParameters Parameters;
+
+	UPROPERTY()
+	TArray<FCachedTrackSource> CacheTracks;
+
+	/** Anonymous array of cleanup functions to perform when a recording has finished */
+	TArray<TFunction<void()>> OnStopCleanup;
+
+	/** Triggered before the recorder is initialized */
+	FOnCacheTrackRecordingPreInitialize OnRecordingPreInitializeEvent;
+
+	/** Triggered when this recorder starts */
+	FOnCacheTrackRecordingStarted OnRecordingStartedEvent;
+
+	/** Triggered when this recorder is stopped */
+	FOnCacheTrackRecordingStopped OnRecordingStoppedEvent;
+
+	/** Triggered when this recorder finishes */
+	FOnCacheTrackRecordingFinished OnRecordingFinishedEvent;
+
+	/** Triggered when this recorder is cancelled */
+	FOnCacheTrackRecordingCancelled OnRecordingCancelledEvent;
+
+	/** Sequencer ptr that controls playback of the desination asset during the recording */
+	TWeakPtr<ISequencer> WeakSequencer;
+
+	/** Due a few ticks after the pre so we are set up with asset creation */
+	int32 NumberOfTicksAfterPre;
+
+	friend class FTickableCacheTrackRecorder;
+
+	/**
+	 * Set the currently active take recorder instance
+	 */
+	static bool SetActiveRecorder(UCacheTrackRecorder* NewActiveRecorder);
+
+	/** Event to trigger when a new recording is initialized */
+	static FOnCacheTrackRecordingInitialized OnRecordingInitializedEvent;
+
+	EAllowEditsMode CachedAllowEditsMode;
+	EAutoChangeMode CachedAutoChangeMode;
+	EUpdateClockSource CachedClockSource;
+	
+	TRange<FFrameNumber> CachedPlaybackRange;
+
+	struct FSavedState
+	{
+		bool bBackedUp = false;
+
+		bool bUseFixedTimeStep = false;
+		double FixedDeltaTime = 1.0;
+	};
+
+	FSavedState SavedState;
+};
diff --git a/Engine/Plugins/VirtualProduction/Takes/Source/TakeTrackRecorders/Public/TrackRecorders/IMovieSceneTrackRecorderFactory.h b/Engine/Plugins/VirtualProduction/Takes/Source/TakeTrackRecorders/Public/TrackRecorders/IMovieSceneTrackRecorderFactory.h
--- a/Engine/Plugins/VirtualProduction/Takes/Source/TakeTrackRecorders/Public/TrackRecorders/IMovieSceneTrackRecorderFactory.h
+++ b/Engine/Plugins/VirtualProduction/Takes/Source/TakeTrackRecorders/Public/TrackRecorders/IMovieSceneTrackRecorderFactory.h
@@ -6,6 +6,9 @@
 #include "Templates/SubclassOf.h"
 #include "Features/IModularFeature.h"
 
+class ISequencer;
+class ULevelSequence;
+class IMovieSceneCachedTrack;
 class UMovieSceneTrackRecorderSettings;
 
 /**  
@@ -56,6 +59,8 @@ public:
 	*/
 	virtual class UMovieSceneTrackRecorder* CreateTrackRecorderForProperty(UObject* InObjectToRecord, const FName& InPropertyToRecord) const = 0;
 
+	virtual class UMovieSceneTrackRecorder* CreateTrackRecorderForCacheTrack(IMovieSceneCachedTrack* CachedTrack, TObjectPtr<ULevelSequence> Sequence, TSharedPtr<ISequencer> Sequencer) const { return nullptr; };
+
 	/**
 	* Get the human readable display name for the recorder. Used for debugging purposes to help identify which factory
 	* is recording a given property.
diff --git a/Engine/Plugins/VirtualProduction/Takes/Source/TakeTrackRecorders/Public/TrackRecorders/MovieSceneTrackRecorder.h b/Engine/Plugins/VirtualProduction/Takes/Source/TakeTrackRecorders/Public/TrackRecorders/MovieSceneTrackRecorder.h
--- a/Engine/Plugins/VirtualProduction/Takes/Source/TakeTrackRecorders/Public/TrackRecorders/MovieSceneTrackRecorder.h
+++ b/Engine/Plugins/VirtualProduction/Takes/Source/TakeTrackRecorders/Public/TrackRecorders/MovieSceneTrackRecorder.h
@@ -127,6 +127,14 @@ public:
 		StopRecordingImpl();
 	}
 
+	/**
+	* Returns false if the track for this recorder has finished recording data
+	*/
+	virtual bool ShouldContinueRecording(const FQualifiedFrameTime& FrameTime) const
+	{
+		return true;
+	}
+
 	void InvalidateObjectToRecord()
 	{
 		ObjectToRecord = nullptr;
diff --git a/Engine/Plugins/VirtualProduction/Takes/Takes.uplugin b/Engine/Plugins/VirtualProduction/Takes/Takes.uplugin
--- a/Engine/Plugins/VirtualProduction/Takes/Takes.uplugin
+++ b/Engine/Plugins/VirtualProduction/Takes/Takes.uplugin
@@ -51,6 +51,11 @@
 			"Name": "TakeSequencer",
 			"Type": "Editor",
 			"LoadingPhase": "Default"
+		},
+		{
+			"Name": "CacheTrackRecorder",
+			"Type": "UncookedOnly",
+			"LoadingPhase": "PostDefault"
 		}
 	]
 }
\ No newline at end of file
diff --git a/Engine/Source/Runtime/MovieScene/Public/Tracks/MovieSceneCachedTrack.h b/Engine/Source/Runtime/MovieScene/Public/Tracks/MovieSceneCachedTrack.h
new file mode 100644
--- /dev/null
+++ b/Engine/Source/Runtime/MovieScene/Public/Tracks/MovieSceneCachedTrack.h
@@ -0,0 +1,35 @@
+﻿// Copyright Epic Games, Inc. All Rights Reserved.
+
+#pragma once
+
+#include "CoreMinimal.h"
+#include "UObject/Interface.h"
+#include "MovieSceneCachedTrack.generated.h"
+
+UINTERFACE()
+class MOVIESCENE_API UMovieSceneCachedTrack : public UInterface
+{
+	GENERATED_BODY()
+};
+
+/**
+ * Can be implemented by tracks that hold cached data. Used by the take recorder plugin to regenerate the cache when recording. 
+ */
+class MOVIESCENE_API IMovieSceneCachedTrack
+{
+	GENERATED_BODY()
+
+public:
+	
+	/**
+	 * Deletes any existing cache data
+	 */
+	virtual void ResetCache() {};
+
+	/**
+	 * Used to enable or disable recording for this track
+	 */
+	virtual void SetCacheRecordingAllowed(bool bShouldRecord) {};
+
+	virtual bool IsCacheRecordingAllowed() const { return true; };
+};
