Take Recorder: Clear out overlay widget when done recording

Speculative fix for GC error after recording with vcam and multi-user

diff --git a/Engine/Plugins/VirtualProduction/Takes/Source/TakeRecorder/Private/Recorder/TakeRecorder.cpp b/Engine/Plugins/VirtualProduction/Takes/Source/TakeRecorder/Private/Recorder/TakeRecorder.cpp
--- a/Engine/Plugins/VirtualProduction/Takes/Source/TakeRecorder/Private/Recorder/TakeRecorder.cpp
+++ b/Engine/Plugins/VirtualProduction/Takes/Source/TakeRecorder/Private/Recorder/TakeRecorder.cpp
@@ -1251,6 +1251,14 @@ void UTakeRecorder::StopInternal(const bool bCancelled)
 		SequenceAsset = nullptr;
 	}
 
+	if (OverlayWidget)
+	{
+		OverlayWidget->RemoveFromParent();
+		OverlayWidget->RemoveFromRoot();
+		OverlayWidget->MarkAsGarbage();
+		OverlayWidget = nullptr;
+	}
+
 	if (SequenceAsset)
 	{
 		UPackage* const Package = SequenceAsset->GetOutermost();
