Adding missing include.

diff --git a/Engine/Plugins/VirtualProduction/Takes/Source/CacheTrackRecorder/Private/Recorder/CacheTrackRecorder.cpp b/Engine/Plugins/VirtualProduction/Takes/Source/CacheTrackRecorder/Private/Recorder/CacheTrackRecorder.cpp
--- a/Engine/Plugins/VirtualProduction/Takes/Source/CacheTrackRecorder/Private/Recorder/CacheTrackRecorder.cpp
+++ b/Engine/Plugins/VirtualProduction/Takes/Source/CacheTrackRecorder/Private/Recorder/CacheTrackRecorder.cpp
@@ -2,6 +2,7 @@
 
 #include "Recorder/CacheTrackRecorder.h"
 
+#include "AnimatedRange.h"
 #include "CoreGlobals.h"
 #include "Engine/Engine.h"
 #include "Engine/EngineTypes.h"
