Asset Definition: Convert Take Preset

diff --git a/Engine/Plugins/VirtualProduction/Takes/Source/TakeRecorder/Private/AssetDefinition_TakePreset.cpp b/Engine/Plugins/VirtualProduction/Takes/Source/TakeRecorder/Private/AssetDefinition_TakePreset.cpp
new file mode 100644
--- /dev/null
+++ b/Engine/Plugins/VirtualProduction/Takes/Source/TakeRecorder/Private/AssetDefinition_TakePreset.cpp
@@ -0,0 +1,54 @@
+// Copyright Epic Games, Inc. All Rights Reserved.
+
+#include "AssetDefinition_TakePreset.h"
+#include "TakePreset.h"
+#include "TakePresetToolkit.h"
+#include "ITakeRecorderModule.h"
+#include "Widgets/STakeRecorderTabContent.h"
+#include "Widgets/DeclarativeSyntaxSupport.h"
+#include "Widgets/SCompoundWidget.h"
+#include "Widgets/Docking/SDockTab.h"
+#include "Framework/Docking/TabManager.h"
+#include "Toolkits/ToolkitManager.h"
+
+#define LOCTEXT_NAMESPACE "AssetTypeActions"
+
+FAssetOpenSupport UAssetDefinition_TakePreset::GetAssetOpenSupport(const FAssetOpenSupportArgs& OpenSupportArgs) const
+{
+	return FAssetOpenSupport(OpenSupportArgs.OpenMethod,OpenSupportArgs.OpenMethod == EAssetOpenMethod::Edit, EToolkitMode::WorldCentric); 
+}
+
+EAssetCommandResult UAssetDefinition_TakePreset::OpenAssets(const FAssetOpenArgs& OpenArgs) const
+{
+	UWorld* WorldContext = nullptr;
+	for (const FWorldContext& Context : GEngine->GetWorldContexts())
+	{
+		if (Context.WorldType == EWorldType::Editor)
+		{
+			WorldContext = Context.World();
+			break;
+		}
+	}
+
+	if (!ensure(WorldContext))
+	{
+		return EAssetCommandResult::Handled;
+	}
+
+	for (UTakePreset* TakePreset : OpenArgs.LoadObjects<UTakePreset>())
+	{
+		TSharedPtr<FTakePresetToolkit> Toolkit = MakeShared<FTakePresetToolkit>();
+		Toolkit->Initialize(OpenArgs.GetToolkitMode(), OpenArgs.ToolkitHost, TakePreset);
+
+		TSharedPtr<SDockTab> DockTab = OpenArgs.ToolkitHost->GetTabManager()->TryInvokeTab(ITakeRecorderModule::TakeRecorderTabName);
+		if (DockTab.IsValid())
+		{
+			TSharedRef<STakeRecorderTabContent> TabContent = StaticCastSharedRef<STakeRecorderTabContent>(DockTab->GetContent());
+			TabContent->SetupForEditing(Toolkit);
+		}
+	}
+
+	return EAssetCommandResult::Handled;
+}
+
+#undef LOCTEXT_NAMESPACE
diff --git a/Engine/Plugins/VirtualProduction/Takes/Source/TakeRecorder/Private/AssetDefinition_TakePreset.h b/Engine/Plugins/VirtualProduction/Takes/Source/TakeRecorder/Private/AssetDefinition_TakePreset.h
new file mode 100644
--- /dev/null
+++ b/Engine/Plugins/VirtualProduction/Takes/Source/TakeRecorder/Private/AssetDefinition_TakePreset.h
@@ -0,0 +1,28 @@
+// Copyright Epic Games, Inc. All Rights Reserved.
+
+#pragma once
+
+#include "TakePreset.h"
+#include "AssetDefinitionDefault.h"
+
+#include "AssetDefinition_TakePreset.generated.h"
+
+UCLASS()
+class UAssetDefinition_TakePreset : public UAssetDefinitionDefault
+{
+	GENERATED_BODY()
+
+public:
+	// UAssetDefinition Begin
+	virtual FText GetAssetDisplayName() const override { return NSLOCTEXT("AssetTypeActions", "AssetTypeActions_TakePreset", "Take Recorder Preset"); }
+	virtual FLinearColor GetAssetColor() const override { return FLinearColor(FColor(226, 155, 72)); }
+	virtual TSoftClassPtr<UObject> GetAssetClass() const override { return UTakePreset::StaticClass(); }
+	virtual TConstArrayView<FAssetCategoryPath> GetAssetCategories() const override
+	{
+		static const auto Categories = { EAssetCategoryPaths::Animation };
+		return Categories;
+	}
+	virtual FAssetOpenSupport GetAssetOpenSupport(const FAssetOpenSupportArgs& OpenSupportArgs) const override;
+	virtual EAssetCommandResult OpenAssets(const FAssetOpenArgs& OpenArgs) const override;
+	// UAssetDefinition End
+};
diff --git a/Engine/Plugins/VirtualProduction/Takes/Source/TakeRecorder/Private/TakePresetActions.cpp b/Engine/Plugins/VirtualProduction/Takes/Source/TakeRecorder/Private/TakePresetActions.cpp
deleted file mode 100644
diff --git a/Engine/Plugins/VirtualProduction/Takes/Source/TakeRecorder/Private/TakePresetActions.h b/Engine/Plugins/VirtualProduction/Takes/Source/TakeRecorder/Private/TakePresetActions.h
deleted file mode 100644
diff --git a/Engine/Plugins/VirtualProduction/Takes/Source/TakeRecorder/Private/TakeRecorderModule.cpp b/Engine/Plugins/VirtualProduction/Takes/Source/TakeRecorder/Private/TakeRecorderModule.cpp
--- a/Engine/Plugins/VirtualProduction/Takes/Source/TakeRecorder/Private/TakeRecorderModule.cpp
+++ b/Engine/Plugins/VirtualProduction/Takes/Source/TakeRecorder/Private/TakeRecorderModule.cpp
@@ -7,7 +7,6 @@
 #include "LevelEditor.h"
 #include "TakeRecorderCommands.h"
 #include "TakeRecorderStyle.h"
-#include "TakePresetActions.h"
 #include "Recorder/TakeRecorder.h"
 
 #include "WorkspaceMenuStructureModule.h"
@@ -271,7 +270,6 @@ void FTakeRecorderModule::StartupModule()
 
 	RegisterDetailCustomizations();
 	RegisterLevelEditorExtensions();
-	RegisterAssetTools();
 	RegisterSettings();
 	RegisterSerializedRecorder();
 
@@ -318,7 +316,6 @@ void FTakeRecorderModule::ShutdownModule()
 
 	UnregisterDetailCustomizations();
 	UnregisterLevelEditorExtensions();
-	UnregisterAssetTools();
 	UnregisterSettings();
 	UnregisterSerializedRecorder();
 }
@@ -420,24 +417,6 @@ void FTakeRecorderModule::UnregisterLevelEditorExtensions()
 	FModuleManager::Get().OnModulesChanged().Remove(ModulesChangedHandle);
 }
 
-void FTakeRecorderModule::RegisterAssetTools()
-{
-	IAssetTools& AssetTools = FModuleManager::LoadModuleChecked<FAssetToolsModule>("AssetTools").Get();
-	{
-		TakePresetActions = MakeShared<FTakePresetActions>();
-		AssetTools.RegisterAssetTypeActions(TakePresetActions.ToSharedRef());
-	}
-}
-
-void FTakeRecorderModule::UnregisterAssetTools()
-{
-	FAssetToolsModule* AssetToolsModule = FModuleManager::GetModulePtr<FAssetToolsModule>("AssetTools");
-	if (AssetToolsModule)
-	{
-		AssetToolsModule->Get().UnregisterAssetTypeActions(TakePresetActions.ToSharedRef());
-	}
-}
-
 void FTakeRecorderModule::RegisterSettings()
 {
 	RegisterSettingsObject(GetMutableDefault<UMovieSceneTakeSettings>());
diff --git a/Engine/Plugins/VirtualProduction/Takes/Source/TakeRecorder/TakeRecorder.Build.cs b/Engine/Plugins/VirtualProduction/Takes/Source/TakeRecorder/TakeRecorder.Build.cs
--- a/Engine/Plugins/VirtualProduction/Takes/Source/TakeRecorder/TakeRecorder.Build.cs
+++ b/Engine/Plugins/VirtualProduction/Takes/Source/TakeRecorder/TakeRecorder.Build.cs
@@ -9,8 +9,8 @@ public class TakeRecorder : ModuleRules
 		PrivateDependencyModuleNames.AddRange(
 			new string[] {
 				"AppFramework",
+				"AssetDefinition",
 				"AssetRegistry",
-				"AssetTools",
 				"ContentBrowser",
 				"Core",
 				"CoreUObject",
