Fixed sequencer caching bugs:
* Fixed issues recording a frame when life cycle track starts at frame 0
* No more inconsistent frames between life cycle track and cache track (I hope, this one is tricky)
* Better Interpolation of the last frame by recording an additional frame at the end
* Life cycle track could drop the first tick sometimes
* Cache moves to the correct position after recording now
* View is restored after recording
* Fixed a crash that could happen when recording gpu systems
* Interpolation is on by default for new tracks
* Better error messages

diff --git a/Engine/Plugins/FX/Niagara/Source/Niagara/Private/MovieScene/MovieSceneNiagaraSystemTrackTemplate.cpp b/Engine/Plugins/FX/Niagara/Source/Niagara/Private/MovieScene/MovieSceneNiagaraSystemTrackTemplate.cpp
--- a/Engine/Plugins/FX/Niagara/Source/Niagara/Private/MovieScene/MovieSceneNiagaraSystemTrackTemplate.cpp
+++ b/Engine/Plugins/FX/Niagara/Source/Niagara/Private/MovieScene/MovieSceneNiagaraSystemTrackTemplate.cpp
@@ -232,7 +232,7 @@ struct FNiagaraSystemUpdateDesiredAgeExecutionToken : IMovieSceneExecutionToken
 
 			if (SystemInstanceController.IsValid() && SystemInstanceController->IsValid() && SystemInstanceController->IsComplete() == false)
 			{
-				const float DesiredAge = float(Context.GetFrameRate().AsSeconds(Context.GetTime() - SpawnSectionStartFrame));
+				const float DesiredAge = FMath::RoundToFloat(static_cast<float>(Context.GetFrameRate().AsSeconds(Context.GetTime() - SpawnSectionStartFrame)) * MinAgeResolution) / MinAgeResolution;
 				if (DesiredAge >= 0)
 				{
 					// Add a quarter of a frame offset here to push the desired age into the middle of the frame since it will be automatically rounded
@@ -252,6 +252,9 @@ struct FNiagaraSystemUpdateDesiredAgeExecutionToken : IMovieSceneExecutionToken
 	ENiagaraSystemSpawnSectionEndBehavior SpawnSectionEndBehavior;
 	ENiagaraAgeUpdateMode AgeUpdateMode;
 	bool bAllowScalability;
+
+	//TODO (mga) move to global sim resolution config, update sim cache as well
+	static constexpr float MinAgeResolution = 10000.0f;
 };
 
 FMovieSceneNiagaraSystemTrackImplementation::FMovieSceneNiagaraSystemTrackImplementation(
diff --git a/Engine/Plugins/FX/Niagara/Source/Niagara/Private/NiagaraComponent.cpp b/Engine/Plugins/FX/Niagara/Source/Niagara/Private/NiagaraComponent.cpp
--- a/Engine/Plugins/FX/Niagara/Source/Niagara/Private/NiagaraComponent.cpp
+++ b/Engine/Plugins/FX/Niagara/Source/Niagara/Private/NiagaraComponent.cpp
@@ -782,57 +782,51 @@ void UNiagaraComponent::TickComponent(float DeltaSeconds, enum ELevelTick TickTy
 		{
 			float AgeDiff = FMath::Max(DesiredAge, 0.0f) - SystemInstanceController->GetAge();
 			int32 TicksToProcess = 0;
-			if (bLockDesiredAgeDeltaTimeToSeekDelta && FMath::Abs(AgeDiff) < KINDA_SMALL_NUMBER)
+
+			if (AgeDiff < 0.0f)
 			{
-				AgeDiff = 0.0f;
+				SystemInstanceController->Reset(FNiagaraSystemInstance::EResetMode::ResetAll);
+				AgeDiff = DesiredAge - SystemInstanceController->GetAge();
 			}
-			else
+
+			if (FMath::Abs(AgeDiff) >= UE_KINDA_SMALL_NUMBER)
 			{
-				if (AgeDiff < 0.0f)
-				{
-					SystemInstanceController->Reset(FNiagaraSystemInstance::EResetMode::ResetAll);
-					AgeDiff = DesiredAge - SystemInstanceController->GetAge();
-				}
-				
-				if (AgeDiff > 0.0f)
+				FNiagaraSystemSimulation* SystemSim = SystemInstanceController->GetSoloSystemSimulation().Get();
+				if (SystemSim)
 				{
-					FNiagaraSystemSimulation* SystemSim = SystemInstanceController->GetSoloSystemSimulation().Get();
-					if (SystemSim)
+					const bool bFixedTickDelta = Asset->HasFixedTickDelta();
+					const float TickDelta = bFixedTickDelta ? Asset->GetFixedTickDeltaTime() : SeekDelta;
+					if (bLockDesiredAgeDeltaTimeToSeekDelta || AgeDiff > TickDelta)
 					{
-						const bool bFixedTickDelta = Asset->HasFixedTickDelta();
-						const float TickDelta = bFixedTickDelta ? Asset->GetFixedTickDeltaTime() : SeekDelta;
-						if (bLockDesiredAgeDeltaTimeToSeekDelta || AgeDiff > TickDelta)
+						// If we're locking the delta time to the seek delta, or we need to seek more than a frame, tick the simulation by the seek delta.
+						const bool bUseMaxSimTime = MaxSimTime > 0.0f;
+						const double EndMaxSimTime = bUseMaxSimTime ? FPlatformTime::Seconds() + MaxSimTime : 0.0f;
+
+						TicksToProcess = FMath::FloorToInt(AgeDiff / TickDelta);
+						while ( TicksToProcess > 0 )
 						{
-							// If we're locking the delta time to the seek delta, or we need to seek more than a frame, tick the simulation by the seek delta.
-							const bool bUseMaxSimTime = MaxSimTime > 0.0f;
-							const double EndMaxSimTime = bUseMaxSimTime ? FPlatformTime::Seconds() + MaxSimTime : 0.0f;
+							// Cannot do multiple tick off the game thread here without additional work. So we pass in null for the completion event which will force GT execution.
+							// If this becomes a perf problem I can add a new path for the tick code to handle multiple ticks.
+							SystemInstanceController->ManualTick(TickDelta, nullptr);
+							--TicksToProcess;
 
-							TicksToProcess = FMath::FloorToInt(AgeDiff / TickDelta);
-							while ( TicksToProcess > 0 )
+							// This limits the amount of time we will consume seeking forward
+							if ( bUseMaxSimTime )
 							{
-								// Cannot do multiple tick off the game thread here without additional work. So we pass in null for the completion event which will force GT execution.
-								// If this becomes a perf problem I can add a new path for the tick code to handle multiple ticks.
-								SystemInstanceController->ManualTick(TickDelta, nullptr);
-								--TicksToProcess;
-
-								// This limits the amount of time we will consume seeking forward
-								if ( bUseMaxSimTime )
+								const double CurrentTime = FPlatformTime::Seconds();
+								if ( CurrentTime >= EndMaxSimTime )
 								{
-									const double CurrentTime = FPlatformTime::Seconds();
-									if ( CurrentTime >= EndMaxSimTime )
-									{
-										break;
-									}
+									break;
 								}
 							}
 						}
-						else
+					}
+					else
+					{
+						// Otherwise just tick by the age difference.
+						if (!bFixedTickDelta || AgeDiff >= TickDelta)
 						{
-							// Otherwise just tick by the age difference.
-							if (!bFixedTickDelta || AgeDiff >= TickDelta)
-							{
-								SystemInstanceController->ManualTick(AgeDiff, nullptr);
-							}
+							SystemInstanceController->ManualTick(AgeDiff, nullptr);
 						}
 					}
 				}
diff --git a/Engine/Plugins/FX/NiagaraSimCaching/Source/NiagaraSimCaching/Private/Niagara/Sequencer/MovieSceneNiagaraCacheTrack.cpp b/Engine/Plugins/FX/NiagaraSimCaching/Source/NiagaraSimCaching/Private/Niagara/Sequencer/MovieSceneNiagaraCacheTrack.cpp
--- a/Engine/Plugins/FX/NiagaraSimCaching/Source/NiagaraSimCaching/Private/Niagara/Sequencer/MovieSceneNiagaraCacheTrack.cpp
+++ b/Engine/Plugins/FX/NiagaraSimCaching/Source/NiagaraSimCaching/Private/Niagara/Sequencer/MovieSceneNiagaraCacheTrack.cpp
@@ -80,6 +80,8 @@ UMovieSceneSection* UMovieSceneNiagaraCacheTrack::CreateNewSection()
 	UMovieSceneNiagaraCacheSection* MovieSceneNiagaraCacheSection = NewObject<UMovieSceneNiagaraCacheSection>(this, NAME_None, RF_Transactional);
 	MovieSceneNiagaraCacheSection->Params.SimCache = NewObject<UNiagaraSimCache>(MovieSceneNiagaraCacheSection, NAME_None, RF_Transactional);
 	MovieSceneNiagaraCacheSection->Params.CacheParameters.AttributeCaptureMode = ENiagaraSimCacheAttributeCaptureMode::RenderingOnly;
+	MovieSceneNiagaraCacheSection->Params.CacheParameters.bAllowInterpolation = true;
+	MovieSceneNiagaraCacheSection->Params.CacheParameters.bAllowVelocityExtrapolation = true;
 	return MovieSceneNiagaraCacheSection;
 }
 
diff --git a/Engine/Plugins/FX/NiagaraSimCaching/Source/NiagaraSimCachingEditor/Private/Sequencer/MovieSceneNiagaraTrackRecorder.cpp b/Engine/Plugins/FX/NiagaraSimCaching/Source/NiagaraSimCachingEditor/Private/Sequencer/MovieSceneNiagaraTrackRecorder.cpp
--- a/Engine/Plugins/FX/NiagaraSimCaching/Source/NiagaraSimCachingEditor/Private/Sequencer/MovieSceneNiagaraTrackRecorder.cpp
+++ b/Engine/Plugins/FX/NiagaraSimCaching/Source/NiagaraSimCachingEditor/Private/Sequencer/MovieSceneNiagaraTrackRecorder.cpp
@@ -8,7 +8,6 @@
 #include "TakeRecorderSettings.h"
 #include "Misc/App.h"
 #include "Misc/CoreMiscDefines.h"
-#include "Misc/CoreMiscDefines.h"
 #include "MovieScene/MovieSceneNiagaraSystemTrack.h"
 #include "Niagara/NiagaraSimCachingEditorPlugin.h"
 #include "Niagara/Sequencer/MovieSceneNiagaraCacheSection.h"
@@ -49,6 +48,7 @@ UMovieSceneTrackRecorder* FMovieSceneNiagaraTrackRecorderFactory::CreateTrackRec
 					if (UNiagaraComponent* NiagaraComponent = Cast<UNiagaraComponent>(Bound))
 					{
 						TrackRecorder->SystemToRecord = NiagaraComponent;
+						NiagaraComponent->SetSimCache(nullptr);
 						break;
 					}
 				}
@@ -67,11 +67,14 @@ UMovieSceneTrackRecorder* FMovieSceneNiagaraTrackRecorderFactory::CreateTrackRec
 						{
 							for (UMovieSceneSection* Section : SystemTrack->GetAllSections())
 							{
+								TRange<FFrameNumber> SectionRange = Section->GetRange();
+								// we move the end of the section one frame out because cache interpolation needs a last frame to work correctly (otherwise the last frame could only be extrapolated)
+								SectionRange.SetUpperBoundValue(SectionRange.GetUpperBoundValue().Value + 1);
 								if (!TrackRecorder->RecordRange.IsSet())
 								{
-									TrackRecorder->RecordRange = Section->GetRange();
+									TrackRecorder->RecordRange = SectionRange;
 								}
-								TrackRecorder->RecordRange = FFrameNumberRange::Hull(*TrackRecorder->RecordRange, Section->GetRange());
+								TrackRecorder->RecordRange = FFrameNumberRange::Hull(*TrackRecorder->RecordRange, SectionRange);
 							}
 						}
 					}
@@ -200,20 +203,6 @@ void UMovieSceneNiagaraTrackRecorder::FinalizeTrackImpl()
 		// finalize the sim cache
 		NiagaraCacheSection->Params.SimCache->EndWrite();
 		
-		// Set the final range 
-		TOptional<TRange<FFrameNumber>> DefaultSectionLength = NiagaraCacheSection->GetAutoSizeRange();
-		if (DefaultSectionLength.IsSet())
-		{
-			if (RecordRange.IsSet())
-			{
-				NiagaraCacheSection->SetRange(TRange<FFrameNumber>(RecordRange->GetLowerBoundValue(), (DefaultSectionLength->GetUpperBoundValue() - DefaultSectionLength->GetLowerBoundValue()) + RecordRange->GetLowerBoundValue()));
-			}
-			else
-			{
-				NiagaraCacheSection->SetRange(DefaultSectionLength.GetValue());
-			}
-		}
-
 		// Activate the section
 		NiagaraCacheSection->SetIsActive(true);
 	}
@@ -224,17 +213,37 @@ void UMovieSceneNiagaraTrackRecorder::RecordSampleImpl(const FQualifiedFrameTime
 	const FFrameRate TickResolution = MovieScene->GetTickResolution();
 	const FFrameNumber CurrentFrame = CurrentFrameTime.ConvertTo(TickResolution).FloorToFrame();
 
-	if (RecordRange.IsSet() && (RecordRange->GetLowerBoundValue() > CurrentFrame || RecordRange->GetUpperBoundValue() < CurrentFrame))
+	if (RecordRange.IsSet() && (RecordRange->GetLowerBoundValue() > CurrentFrame || RecordRange->GetUpperBoundValue() <= CurrentFrame))
 	{
 		return;
 	}
 
 	if (NiagaraCacheSection.IsValid() && SystemToRecord.IsValid())
 	{
-		if ( NiagaraCacheSection->Params.SimCache->WriteFrame(SystemToRecord.Get()) )
+		FNiagaraSimCacheFeedbackContext FeedbackContext;
+		FeedbackContext.bAutoLogIssues = false;
+		if (NiagaraCacheSection->Params.SimCache->WriteFrame(SystemToRecord.Get(), FeedbackContext))
 		{
+			if (!bRecordedFirstFrame)
+			{
+				// set to the actual first recorded frame, because systems with spawn rate can tick for a few frames without having particles
+				NiagaraCacheSection->SetStartFrame(CurrentFrame);
+				bRecordedFirstFrame = true;
+			}
+			
 			// Expand the section to the new length
-			NiagaraCacheSection->SetEndFrame(CurrentFrame);
+			FFrameNumber EndFrame = CurrentFrameTime.Time.CeilToFrame();
+			NiagaraCacheSection->ExpandToFrame(EndFrame);
+		}
+		
+		for (const FString& Warning : FeedbackContext.Warnings)
+		{
+			UE_LOG(LogNiagaraSimCachingEditor, Warning, TEXT("Recording sim cache for frame %i: %s"), CurrentFrame.Value, *Warning);
+		}
+
+		for (const FString& Error : FeedbackContext.Errors)
+		{
+			UE_LOG(LogNiagaraSimCachingEditor, Warning, TEXT("Unable to record sim cache for frame %i: %s"), CurrentFrame.Value, *Error);
 		}
 	}
 }
diff --git a/Engine/Plugins/FX/NiagaraSimCaching/Source/NiagaraSimCachingEditor/Public/Sequencer/MovieSceneNiagaraTrackRecorder.h b/Engine/Plugins/FX/NiagaraSimCaching/Source/NiagaraSimCachingEditor/Public/Sequencer/MovieSceneNiagaraTrackRecorder.h
--- a/Engine/Plugins/FX/NiagaraSimCaching/Source/NiagaraSimCachingEditor/Public/Sequencer/MovieSceneNiagaraTrackRecorder.h
+++ b/Engine/Plugins/FX/NiagaraSimCaching/Source/NiagaraSimCachingEditor/Public/Sequencer/MovieSceneNiagaraTrackRecorder.h
@@ -69,6 +69,7 @@ private:
 	FFrameNumber RecordStartFrame;
 
 	TOptional<FFrameNumberRange> RecordRange;
+	bool bRecordedFirstFrame = false;
 
 	friend FMovieSceneNiagaraTrackRecorderFactory;
 };
diff --git a/Engine/Plugins/VirtualProduction/Takes/Source/CacheTrackRecorder/Private/Recorder/CacheTrackRecorder.cpp b/Engine/Plugins/VirtualProduction/Takes/Source/CacheTrackRecorder/Private/Recorder/CacheTrackRecorder.cpp
--- a/Engine/Plugins/VirtualProduction/Takes/Source/CacheTrackRecorder/Private/Recorder/CacheTrackRecorder.cpp
+++ b/Engine/Plugins/VirtualProduction/Takes/Source/CacheTrackRecorder/Private/Recorder/CacheTrackRecorder.cpp
@@ -456,11 +456,9 @@ bool UCacheTrackRecorder::Initialize(ULevelSequence* LevelSequenceBase, const TA
 	{
 		// If a start frame was specified, adjust the playback range before rewinding to the beginning of the playback range
 		UMovieScene* MovieScene = SequenceAsset->GetMovieScene();
+		CachedViewRange = MovieScene->GetEditorData().GetViewRange();
 		MovieScene->SetPlaybackRange(TRange<FFrameNumber>(Parameters.StartFrame, MovieScene->GetPlaybackRange().GetUpperBoundValue()));
 
-		// Always start the recording at the beginning of the playback range
-		Sequencer->SetLocalTime(MovieScene->GetPlaybackRange().GetLowerBoundValue());
-
 		// Center the view range around the current time about to be captured
 		FAnimatedRange Range = Sequencer->GetViewRange();
 		FTimecode CurrentTime = FApp::GetTimecode();
@@ -640,7 +638,7 @@ void UCacheTrackRecorder::Tick(float DeltaTime)
 	{
 		NumberOfTicksAfterPre = 0;
 		Start();
-		InternalTick(0.0f);
+		InternalTick(DeltaTime);
 	}
 	else if (State == ECacheTrackRecorderState::Started)
 	{
@@ -697,15 +695,6 @@ void UCacheTrackRecorder::InternalTick(float DeltaTime)
 		{
 			ShouldContinue = true;
 			CacheSource.Recorder->RecordSample(RecordTime);
-			if (UMovieSceneTrack* SceneTrack = Cast<UMovieSceneTrack>(CacheSource.Track))
-			{
-				FFrameNumber EndFrame = RecordTime.Time.CeilToFrame();
-				for (UMovieSceneSection* SubSection : SceneTrack->GetAllSections())
-				{
-					// expand section range to this recording time
-					SubSection->ExpandToFrame(EndFrame);
-				}
-			}
 		}
 	}
 	if (ShouldContinue == false && DeltaTime > 0.0f)
@@ -798,24 +787,6 @@ void UCacheTrackRecorder::PreRecord()
 		};
 		OnStopCleanup.Add(RestoreEditorTick);
 	}
-}
-
-void UCacheTrackRecorder::Start()
-{
-	FTimecode Timecode = FApp::GetTimecode();
-
-	State = ECacheTrackRecorderState::Started;
-
-	TSharedPtr<ISequencer> Sequencer = WeakSequencer.Pin();
-
-	CurrentFrameTime = FFrameTime(0);
-	TimecodeAtStart = Timecode;
-
-	// Discard any entity tokens we have so that restore state does not take effect when we delete any sections that recording will be replacing.
-	if (Sequencer.IsValid())
-	{
-		Sequencer->PreAnimatedState.DiscardEntityTokens();
-	}
 
 	if (UMovieScene* MovieScene = SequenceAsset->GetMovieScene())
 	{
@@ -829,7 +800,7 @@ void UCacheTrackRecorder::Start()
 
 		FFrameRate FrameRate = MovieScene->GetDisplayRate();
 		FFrameRate TickResolution = MovieScene->GetTickResolution();
-		FFrameNumber PlaybackStartFrame = Parameters.Project.bStartAtCurrentTimecode ? FFrameRate::TransformTime(FFrameTime(Timecode.ToFrameNumber(FrameRate)), FrameRate, TickResolution).FloorToFrame() : MovieScene->GetPlaybackRange().GetLowerBoundValue();
+		FFrameNumber PlaybackStartFrame = Parameters.Project.bStartAtCurrentTimecode ? FFrameRate::TransformTime(FFrameTime(FApp::GetTimecode().ToFrameNumber(FrameRate)), FrameRate, TickResolution).FloorToFrame() : MovieScene->GetPlaybackRange().GetLowerBoundValue();
 
 		// Transform all the sections to start the playback start frame
 		FFrameNumber DeltaFrame = PlaybackStartFrame - MovieScene->GetPlaybackRange().GetLowerBoundValue();
@@ -851,9 +822,29 @@ void UCacheTrackRecorder::Start()
 		if (Sequencer.IsValid())
 		{
 			Sequencer->SetGlobalTime(PlaybackStartFrame);
-			Sequencer->SetPlaybackStatus(EMovieScenePlayerStatus::Playing);
+			Sequencer->SetPlaybackStatus(EMovieScenePlayerStatus::Paused);
 		}
 	}
+}
+
+void UCacheTrackRecorder::Start()
+{
+	FTimecode Timecode = FApp::GetTimecode();
+
+	State = ECacheTrackRecorderState::Started;
+
+	TSharedPtr<ISequencer> Sequencer = WeakSequencer.Pin();
+
+	CurrentFrameTime = FFrameTime(0);
+	TimecodeAtStart = Timecode;
+
+	// Discard any entity tokens we have so that restore state does not take effect when we delete any sections that recording will be replacing.
+	if (Sequencer.IsValid())
+	{
+		Sequencer->PreAnimatedState.DiscardEntityTokens();
+	}
+
+	Sequencer->SetPlaybackStatus(EMovieScenePlayerStatus::Playing);
 
 	FQualifiedFrameTime RecordTime = GetRecordTime();
 	FFrameNumber RecordStartFrame = RecordTime.Time.FloorToFrame();
@@ -866,8 +857,6 @@ void UCacheTrackRecorder::Start()
 		{
 			CacheSource.Recorder->SetSectionStartTimecode(CurrentTimecode, RecordTime.Time.FloorToFrame());
 
-			// Record immediately so that there's a key on the first frame of recording
-			CacheSource.Recorder->RecordSample(RecordTime);
 			if (UMovieSceneTrack* SceneTrack = Cast<UMovieSceneTrack>(CacheSource.Track))
 			{
 				for (UMovieSceneSection* SubSection : SceneTrack->GetAllSections())
@@ -996,10 +985,11 @@ void UCacheTrackRecorder::StopInternal(const bool bCancelled)
 			}
 		}
 
-		// Restore the playback range to what it was before recording.
+		// Restore the playback/view range to what it was before recording.
 		if (MovieScene)
 		{
 			MovieScene->SetPlaybackRange(CachedPlaybackRange);
+			Sequencer->SetViewRange(CachedViewRange, EViewRangeInterpolation::Immediate);
 		}
 
 		if (bRecordingFinished)
diff --git a/Engine/Plugins/VirtualProduction/Takes/Source/CacheTrackRecorder/Public/Recorder/CacheTrackRecorder.h b/Engine/Plugins/VirtualProduction/Takes/Source/CacheTrackRecorder/Public/Recorder/CacheTrackRecorder.h
--- a/Engine/Plugins/VirtualProduction/Takes/Source/CacheTrackRecorder/Public/Recorder/CacheTrackRecorder.h
+++ b/Engine/Plugins/VirtualProduction/Takes/Source/CacheTrackRecorder/Public/Recorder/CacheTrackRecorder.h
@@ -344,6 +344,8 @@ private:
 
 	/** Due a few ticks after the pre so we are set up with asset creation */
 	int32 NumberOfTicksAfterPre;
+	
+	TRange<double> CachedViewRange;
 
 	friend class FTickableCacheTrackRecorder;
 
