Take Recorder: Timecode setting names now default to the same as Anim Sequence so that aligning works

diff --git a/Engine/Plugins/VirtualProduction/Takes/Source/TakeMovieScene/Private/MovieSceneTakeSettings.cpp b/Engine/Plugins/VirtualProduction/Takes/Source/TakeMovieScene/Private/MovieSceneTakeSettings.cpp
--- a/Engine/Plugins/VirtualProduction/Takes/Source/TakeMovieScene/Private/MovieSceneTakeSettings.cpp
+++ b/Engine/Plugins/VirtualProduction/Takes/Source/TakeMovieScene/Private/MovieSceneTakeSettings.cpp
@@ -7,12 +7,12 @@
 #define LOCTEXT_NAMESPACE "MovieSceneTakeSettings"
 
 UMovieSceneTakeSettings::UMovieSceneTakeSettings()
-	: HoursName(TEXT("Hours"))
-	, MinutesName(TEXT("Minutes"))
-	, SecondsName(TEXT("Seconds"))
-	, FramesName(TEXT("Frames"))
-	, SubFramesName(TEXT("SubFrames"))
-	, SlateName(TEXT("Slate"))
+	: HoursName(TEXT("TCHour"))
+	, MinutesName(TEXT("TCMinute"))
+	, SecondsName(TEXT("TCSecond"))
+	, FramesName(TEXT("TCFrame"))
+	, SubFramesName(TEXT("TCSubframe"))
+	, SlateName(TEXT("TCSlate"))
 {
 }
 
