Take Recorder: Fix styles for toggle buttons

#rnx

diff --git a/Engine/Plugins/VirtualProduction/Takes/Source/TakeRecorder/Private/TakeRecorderStyle.cpp b/Engine/Plugins/VirtualProduction/Takes/Source/TakeRecorder/Private/TakeRecorderStyle.cpp
--- a/Engine/Plugins/VirtualProduction/Takes/Source/TakeRecorder/Private/TakeRecorderStyle.cpp
+++ b/Engine/Plugins/VirtualProduction/Takes/Source/TakeRecorder/Private/TakeRecorderStyle.cpp
@@ -5,6 +5,7 @@
 #include "Styling/SlateTypes.h"
 #include "Styling/CoreStyle.h"
 #include "Styling/AppStyle.h"
+#include "Styling/StyleColors.h"
 
 FName FTakeRecorderStyle::StyleName("TakeRecorderStyle");
 
@@ -26,11 +27,6 @@ FTakeRecorderStyle::FTakeRecorderStyle()
 	const FLinearColor DarkGrey( FColor(30, 30, 30) );
 	const FLinearColor Black(FLinearColor::Black);
 
-	const FLinearColor SelectionColor( 0.728f, 0.364f, 0.003f );
-	const FLinearColor SelectionColor_Subdued( 0.807f, 0.596f, 0.388f );
-	const FLinearColor SelectionColor_Inactive( 0.25f, 0.25f, 0.25f );
-	const FLinearColor SelectionColor_Pressed( 0.701f, 0.225f, 0.003f );
-
 	SetContentRoot(FPaths::EnginePluginsDir() / TEXT("VirtualProduction/Takes/Content"));
 	SetCoreContentRoot(FPaths::EngineContentDir() / TEXT("Slate"));
 
@@ -105,26 +101,26 @@ FTakeRecorderStyle::FTakeRecorderStyle()
 	const FCheckBoxStyle ToggleButtonStyle = FCheckBoxStyle()
 		.SetCheckBoxType(ESlateCheckBoxType::ToggleButton)
 		.SetUncheckedImage( FSlateNoResource() )
-		.SetUncheckedHoveredImage( FSlateBoxBrush(RootToCoreContentDir(TEXT("Common/RoundedSelection_16x.png")), 4.0f/16.0f, SelectionColor ) )
-		.SetUncheckedPressedImage( FSlateBoxBrush(RootToCoreContentDir(TEXT("Common/RoundedSelection_16x.png")), 4.0f/16.0f, SelectionColor_Pressed ) )
-		.SetCheckedImage( FSlateBoxBrush(RootToCoreContentDir(TEXT("Common/RoundedSelection_16x.png")),  4.0f/16.0f, SelectionColor_Pressed ) )
-		.SetCheckedHoveredImage( FSlateBoxBrush(RootToCoreContentDir(TEXT("Common/RoundedSelection_16x.png")),  4.0f/16.0f, SelectionColor ) )
-		.SetCheckedPressedImage( FSlateBoxBrush(RootToCoreContentDir(TEXT("Common/RoundedSelection_16x.png")),  4.0f/16.0f, SelectionColor_Pressed ) )
-		.SetUndeterminedImage( FSlateBoxBrush(RootToCoreContentDir(TEXT("Common/RoundedSelection_16x.png")),  4.0f/16.0f, SelectionColor_Inactive) )
-		.SetUndeterminedHoveredImage( FSlateBoxBrush(RootToCoreContentDir(TEXT("Common/RoundedSelection_16x.png")),  4.0f/16.0f, SelectionColor ) )
-		.SetUndeterminedPressedImage( FSlateBoxBrush(RootToCoreContentDir(TEXT("Common/RoundedSelection_16x.png")),  4.0f/16.0f, SelectionColor_Inactive) );
+		.SetUncheckedHoveredImage( FSlateBoxBrush(RootToCoreContentDir(TEXT("Common/RoundedSelection_16x.png")), 4.0f/16.0f, FStyleColors::PrimaryHover) )
+		.SetUncheckedPressedImage( FSlateBoxBrush(RootToCoreContentDir(TEXT("Common/RoundedSelection_16x.png")), 4.0f/16.0f, FStyleColors::PrimaryPress) )
+		.SetCheckedImage( FSlateBoxBrush(RootToCoreContentDir(TEXT("Common/RoundedSelection_16x.png")),  4.0f/16.0f, FStyleColors::PrimaryPress) )
+		.SetCheckedHoveredImage( FSlateBoxBrush(RootToCoreContentDir(TEXT("Common/RoundedSelection_16x.png")),  4.0f/16.0f, FStyleColors::PrimaryHover) )
+		.SetCheckedPressedImage( FSlateBoxBrush(RootToCoreContentDir(TEXT("Common/RoundedSelection_16x.png")),  4.0f/16.0f, FStyleColors::PrimaryPress) )
+		.SetUndeterminedImage( FSlateBoxBrush(RootToCoreContentDir(TEXT("Common/RoundedSelection_16x.png")),  4.0f/16.0f, FStyleColors::SelectInactive) )
+		.SetUndeterminedHoveredImage( FSlateBoxBrush(RootToCoreContentDir(TEXT("Common/RoundedSelection_16x.png")),  4.0f/16.0f, FStyleColors::SelectInactive) )
+		.SetUndeterminedPressedImage( FSlateBoxBrush(RootToCoreContentDir(TEXT("Common/RoundedSelection_16x.png")),  4.0f/16.0f, FStyleColors::SelectInactive) );
 
 	const FCheckBoxStyle ToggleButtonIndicatorStyle = FCheckBoxStyle()
 		.SetCheckBoxType(ESlateCheckBoxType::ToggleButton)
 		.SetUncheckedImage( FSlateNoResource() )
 		.SetUncheckedHoveredImage( FSlateNoResource() )
 		.SetUncheckedPressedImage( FSlateNoResource() )
-		.SetCheckedImage( FSlateBoxBrush(RootToCoreContentDir(TEXT("Common/RoundedSelection_16x.png")),  4.0f/16.0f, SelectionColor_Pressed ) )
-		.SetCheckedHoveredImage( FSlateBoxBrush(RootToCoreContentDir(TEXT("Common/RoundedSelection_16x.png")),  4.0f/16.0f, SelectionColor_Pressed ) )
-		.SetCheckedPressedImage( FSlateBoxBrush(RootToCoreContentDir(TEXT("Common/RoundedSelection_16x.png")),  4.0f/16.0f, SelectionColor_Pressed ) )
-		.SetUndeterminedImage( FSlateBoxBrush(RootToCoreContentDir(TEXT("Common/RoundedSelection_16x.png")),  4.0f/16.0f, SelectionColor_Inactive) )
-		.SetUndeterminedHoveredImage( FSlateBoxBrush(RootToCoreContentDir(TEXT("Common/RoundedSelection_16x.png")),  4.0f/16.0f, SelectionColor_Inactive) )
-		.SetUndeterminedPressedImage( FSlateBoxBrush(RootToCoreContentDir(TEXT("Common/RoundedSelection_16x.png")),  4.0f/16.0f, SelectionColor_Inactive) );
+		.SetCheckedImage( FSlateBoxBrush(RootToCoreContentDir(TEXT("Common/RoundedSelection_16x.png")),  4.0f/16.0f, FStyleColors::PrimaryPress) )
+		.SetCheckedHoveredImage( FSlateBoxBrush(RootToCoreContentDir(TEXT("Common/RoundedSelection_16x.png")),  4.0f/16.0f, FStyleColors::PrimaryHover) )
+		.SetCheckedPressedImage( FSlateBoxBrush(RootToCoreContentDir(TEXT("Common/RoundedSelection_16x.png")),  4.0f/16.0f, FStyleColors::PrimaryPress) )
+		.SetUndeterminedImage( FSlateBoxBrush(RootToCoreContentDir(TEXT("Common/RoundedSelection_16x.png")),  4.0f/16.0f, FStyleColors::SelectInactive) )
+		.SetUndeterminedHoveredImage( FSlateBoxBrush(RootToCoreContentDir(TEXT("Common/RoundedSelection_16x.png")),  4.0f/16.0f, FStyleColors::SelectInactive) )
+		.SetUndeterminedPressedImage( FSlateBoxBrush(RootToCoreContentDir(TEXT("Common/RoundedSelection_16x.png")),  4.0f/16.0f, FStyleColors::SelectInactive) );
 
 	Set( "ToggleButtonCheckbox", ToggleButtonStyle );	
 	Set( "ToggleButtonIndicatorCheckbox", ToggleButtonIndicatorStyle );	
@@ -158,7 +154,7 @@ FTakeRecorderStyle::FTakeRecorderStyle()
 	Set("ClassIcon.TakeRecorderCameraCutSource", new FSlateImageBrush(RootToContentDir(TEXT("TakeRecorderCameraCutSource_16x.png")), Icon16x16));
 	Set("ClassThumbnail.TakeRecorderCameraCutSource", new FSlateImageBrush(RootToContentDir(TEXT("TakeRecorderCameraCutSource_64x.png")), Icon64x64));
 
-	Set("ClassIcon.TakeRecorderLevelSequenceSource", new FSlateImageBrush(RootToContentDir(TEXT("TakeRecorderLevelSequenceSource_16x.png")), Icon16x16));
+	Set("ClassIcon.TakeRecorderLevelSequenceSource", new FSlateVectorImageBrush(RootToContentDir(TEXT("AssetIcons/TakeRecorderLevelSequenceSource_16.svg")), Icon16x16));
 	Set("ClassThumbnail.TakeRecorderLevelSequenceSource", new FSlateVectorImageBrush(RootToContentDir(TEXT("AssetIcons/TakeRecorderLevelSequenceSource_64.svg")), Icon64x64));
 
 	Set("ClassIcon.TakeRecorderLiveLinkSource", new FSlateVectorImageBrush(RootToContentDir(TEXT("AssetIcons/TakeRecorderLiveLinkSource_64.svg")), Icon16x16));
@@ -195,7 +191,7 @@ FTakeRecorderStyle::FTakeRecorderStyle()
 	Set("TakeRecorder.SequencerButton", new FSlateImageBrush(RootToContentDir(TEXT("Sequencer.png")), Icon20x20));
 	Set("TakeRecorder.ReviewRecordingButton", new FSlateImageBrush(RootToContentDir(TEXT("ReviewRecording.png")), Icon20x20));
 	Set("TakeRecorder.MarkFrame", new FSlateImageBrush(RootToContentDir(TEXT("MarkFrame.png")), Icon20x20));
-	Set("TakeRecorder.SequenceToRecordIntoButton", new FSlateImageBrush(RootToContentDir(TEXT("SequenceToRecordInto.png")), Icon20x20));
+	Set("TakeRecorder.SequenceToRecordIntoButton", new FSlateVectorImageBrush(RootToContentDir(TEXT("SequenceToRecordInto.svg")), Icon20x20));
 
 	Set("TakeRecorder.Slate", new FSlateColorBrush(MediumGrey));
 	Set("TakeRecorder.Slate.ClapperBackground", new FSlateColorBrush(AlmostWhite));
