Add context menu option to record multiple cache tracks at once

diff --git a/Engine/Plugins/FX/NiagaraSimCaching/Source/NiagaraSimCachingEditor/Private/Sequencer/NiagaraCacheTrackEditor.cpp b/Engine/Plugins/FX/NiagaraSimCaching/Source/NiagaraSimCachingEditor/Private/Sequencer/NiagaraCacheTrackEditor.cpp
--- a/Engine/Plugins/FX/NiagaraSimCaching/Source/NiagaraSimCachingEditor/Private/Sequencer/NiagaraCacheTrackEditor.cpp
+++ b/Engine/Plugins/FX/NiagaraSimCaching/Source/NiagaraSimCachingEditor/Private/Sequencer/NiagaraCacheTrackEditor.cpp
@@ -348,6 +348,18 @@ void FNiagaraCacheTrackEditor::BuildTrackContextMenu(FMenuBuilder& MenuBuilder,
 					)
 				)
 			);
+
+			MenuBuilder.AddMenuEntry(
+				LOCTEXT("RecordCache", "Record cache"),
+				LOCTEXT("RecordCacheToolTip", "Refresh the data in the selected cache tracks by playing the sequence and recording the Niagara system."),
+				FSlateIcon(),
+				FUIAction(
+					FExecuteAction::CreateLambda([this]()
+					{
+						UCacheTrackRecorder::RecordSelectedTracks(GetSequencer());
+					})
+				)
+			);
 		}
 	}
 
diff --git a/Engine/Plugins/VirtualProduction/Takes/Source/CacheTrackRecorder/Private/Recorder/CacheTrackRecorder.cpp b/Engine/Plugins/VirtualProduction/Takes/Source/CacheTrackRecorder/Private/Recorder/CacheTrackRecorder.cpp
--- a/Engine/Plugins/VirtualProduction/Takes/Source/CacheTrackRecorder/Private/Recorder/CacheTrackRecorder.cpp
+++ b/Engine/Plugins/VirtualProduction/Takes/Source/CacheTrackRecorder/Private/Recorder/CacheTrackRecorder.cpp
@@ -317,8 +317,18 @@ FOnCacheTrackRecordingInitialized& UCacheTrackRecorder::OnRecordingInitialized()
 
 void UCacheTrackRecorder::RecordCacheTrack(IMovieSceneCachedTrack* Track, TSharedPtr<ISequencer> Sequencer, FCacheRecorderParameters Parameters)
 {	
+	if (Track)
+	{
+		TArray<IMovieSceneCachedTrack*> CacheTracks;
+		CacheTracks.Add(Track);
+		RecordCacheTracks(CacheTracks, Sequencer, Parameters);
+	}
+}
+
+void UCacheTrackRecorder::RecordCacheTracks(const TArray<IMovieSceneCachedTrack*>& CacheTracks, TSharedPtr<ISequencer> Sequencer, FCacheRecorderParameters Parameters)
+{
 	ULevelSequence* LevelSequence = Sequencer ? Cast<ULevelSequence>(Sequencer->GetFocusedMovieSceneSequence()) : nullptr;	
-	if (LevelSequence && Track)
+	if (LevelSequence && CacheTracks.Num() > 0)
 	{
 		Parameters.StartFrame = LevelSequence->GetMovieScene()->GetPlaybackRange().GetLowerBoundValue();
 		Parameters.User.CountdownSeconds = 0;
@@ -331,9 +341,6 @@ void UCacheTrackRecorder::RecordCacheTrack(IMovieSceneCachedTrack* Track, TShare
 			Parameters.StartFrame = Sequencer->GetLocalTime().Time.FrameNumber;
 		}
 
-		TArray<IMovieSceneCachedTrack*> CacheTracks;
-		CacheTracks.Add(Track);
-
 		UCacheTrackRecorder* NewRecorder = NewObject<UCacheTrackRecorder>(GetTransientPackage(), NAME_None, RF_Transient);
 		UTakeMetaData* TakeMetaData = LevelSequence->FindOrAddMetaData<UTakeMetaData>();
 		if (TakeMetaData->GetSlate().IsEmpty())
@@ -354,6 +361,24 @@ void UCacheTrackRecorder::RecordCacheTrack(IMovieSceneCachedTrack* Track, TShare
 	}
 }
 
+void UCacheTrackRecorder::RecordSelectedTracks(TSharedPtr<ISequencer> Sequencer, FCacheRecorderParameters Parameters)
+{
+	TArray<UMovieSceneTrack*> SelectedTracks;
+	Sequencer->GetSelectedTracks(SelectedTracks);
+
+	// gather all selected cache tracks to record
+	TArray<IMovieSceneCachedTrack*> CacheTracks;
+	for (UMovieSceneTrack* SelectedTrack : SelectedTracks)
+	{
+		if (IMovieSceneCachedTrack* CacheTrack = Cast<IMovieSceneCachedTrack>(SelectedTrack))
+		{
+			CacheTracks.Add(CacheTrack);
+		}
+	}
+
+	RecordCacheTracks(CacheTracks, Sequencer, Parameters);
+}
+
 bool UCacheTrackRecorder::SetActiveRecorder(UCacheTrackRecorder* NewActiveRecorder)
 {
 	if (GetCurrentRecorder().IsValid())
diff --git a/Engine/Plugins/VirtualProduction/Takes/Source/CacheTrackRecorder/Public/Recorder/CacheTrackRecorder.h b/Engine/Plugins/VirtualProduction/Takes/Source/CacheTrackRecorder/Public/Recorder/CacheTrackRecorder.h
--- a/Engine/Plugins/VirtualProduction/Takes/Source/CacheTrackRecorder/Public/Recorder/CacheTrackRecorder.h
+++ b/Engine/Plugins/VirtualProduction/Takes/Source/CacheTrackRecorder/Public/Recorder/CacheTrackRecorder.h
@@ -154,6 +154,8 @@ public:
 	 * Utility method to quickly record a single cache track
 	 */
 	static void RecordCacheTrack(IMovieSceneCachedTrack* Track, TSharedPtr<ISequencer> Sequencer, FCacheRecorderParameters Parameters = FCacheRecorderParameters());
+	static void RecordCacheTracks(const TArray<IMovieSceneCachedTrack*>& CacheTracks, TSharedPtr<ISequencer> Sequencer, FCacheRecorderParameters Parameters = FCacheRecorderParameters());
+	static void RecordSelectedTracks(TSharedPtr<ISequencer> Sequencer, FCacheRecorderParameters Parameters = FCacheRecorderParameters());
 public:
 
 	/**
