Take Recorder: If recording with the Timecode clock source, and the user has toggled off Start at Current Timecode, override it. Otherwise, the sequence would end up starting at 0, but the recording is actually at the current timecode (ie. 10:20:30:40)

diff --git a/Engine/Plugins/VirtualProduction/Takes/Source/TakeRecorder/Private/Recorder/TakeRecorder.cpp b/Engine/Plugins/VirtualProduction/Takes/Source/TakeRecorder/Private/Recorder/TakeRecorder.cpp
--- a/Engine/Plugins/VirtualProduction/Takes/Source/TakeRecorder/Private/Recorder/TakeRecorder.cpp
+++ b/Engine/Plugins/VirtualProduction/Takes/Source/TakeRecorder/Private/Recorder/TakeRecorder.cpp
@@ -491,6 +491,12 @@ bool UTakeRecorder::Initialize ( ULevelSequence* LevelSequenceBase, UTakeRecorde
 		Parameters.User.bAutoLock = false;
 	}
 
+	if (Parameters.Project.RecordingClockSource == EUpdateClockSource::Timecode)
+	{
+		Parameters.Project.bStartAtCurrentTimecode = true;
+		UE_LOG(LogTakesCore, Warning, TEXT("Overriding Start at Current Timecode to True since the clock source is synced to Timecode."));
+	}
+
 	// Perform any other parameter-configurable initialization. Must have a valid world at this point.
 	InitializeFromParameters();
 
