Take Recorder: Save engine time dilation when changed through the cockpit

#rnx

diff --git a/Engine/Plugins/VirtualProduction/Takes/Source/TakeRecorder/Private/Widgets/STakeRecorderCockpit.cpp b/Engine/Plugins/VirtualProduction/Takes/Source/TakeRecorder/Private/Widgets/STakeRecorderCockpit.cpp
--- a/Engine/Plugins/VirtualProduction/Takes/Source/TakeRecorder/Private/Widgets/STakeRecorderCockpit.cpp
+++ b/Engine/Plugins/VirtualProduction/Takes/Source/TakeRecorder/Private/Widgets/STakeRecorderCockpit.cpp
@@ -964,6 +964,7 @@ float STakeRecorderCockpit::GetEngineTimeDilation() const
 void STakeRecorderCockpit::SetEngineTimeDilation(float InEngineTimeDilation)
 {
 	GetMutableDefault<UTakeRecorderUserSettings>()->Settings.EngineTimeDilation = InEngineTimeDilation;
+	GetMutableDefault<UTakeRecorderUserSettings>()->SaveConfig();
 }
 
 FReply STakeRecorderCockpit::OnAddMarkedFrame()
