Take Recorder: Remove setting the playback range to the extent of all the sections

Retested UE-63870 and animations came out to be the same duration as the recording.

diff --git a/Engine/Plugins/VirtualProduction/Takes/Source/TakeRecorder/Private/Recorder/TakeRecorder.cpp b/Engine/Plugins/VirtualProduction/Takes/Source/TakeRecorder/Private/Recorder/TakeRecorder.cpp
--- a/Engine/Plugins/VirtualProduction/Takes/Source/TakeRecorder/Private/Recorder/TakeRecorder.cpp
+++ b/Engine/Plugins/VirtualProduction/Takes/Source/TakeRecorder/Private/Recorder/TakeRecorder.cpp
@@ -1163,11 +1163,6 @@ void UTakeRecorder::StopInternal(const bool bCancelled)
 				MovieScene->SetPlaybackRange(CachedPlaybackRange);
 			}
 		}
-		else
-		{
-			const bool bUpperBoundOnly = true; // Only expand the upper bound because the lower bound should have been set at the start of recording and should not change if there's existing data before the start
-			TakesUtils::ClampPlaybackRangeToEncompassAllSections(SequenceAsset->GetMovieScene(), bUpperBoundOnly);
-		}
 
 		if (bRecordingFinished)
 		{
diff --git a/Engine/Plugins/VirtualProduction/Takes/Source/TakeRecorderSources/Private/TakeRecorderActorSource.cpp b/Engine/Plugins/VirtualProduction/Takes/Source/TakeRecorderSources/Private/TakeRecorderActorSource.cpp
--- a/Engine/Plugins/VirtualProduction/Takes/Source/TakeRecorderSources/Private/TakeRecorderActorSource.cpp
+++ b/Engine/Plugins/VirtualProduction/Takes/Source/TakeRecorderSources/Private/TakeRecorderActorSource.cpp
@@ -813,9 +813,6 @@ TArray<UTakeRecorderSource*> UTakeRecorderActorSource::PostRecording(ULevelSeque
 			ProcessRecordedTimes(InSequence);
 		}
 
-		// Expand the Movie Scene Playback Range to encompass all of the sections now that they've all been created.
-		SequenceRecorderUtils::ExtendSequencePlaybackRange(InSequence);
-
 		if (Target.IsValid() && !ParentSource)
 		{
 			// Automatically add or update the camera cut track if there is a camera component
