Remove unintentionl pragma once in cpp file. Horde issue 293027.

diff --git a/Engine/Plugins/VirtualProduction/Takes/Source/TakeRecorder/Private/Widgets/TakeRecorderAudioSettingsCustomization.cpp b/Engine/Plugins/VirtualProduction/Takes/Source/TakeRecorder/Private/Widgets/TakeRecorderAudioSettingsCustomization.cpp
--- a/Engine/Plugins/VirtualProduction/Takes/Source/TakeRecorder/Private/Widgets/TakeRecorderAudioSettingsCustomization.cpp
+++ b/Engine/Plugins/VirtualProduction/Takes/Source/TakeRecorder/Private/Widgets/TakeRecorderAudioSettingsCustomization.cpp
@@ -1,7 +1,5 @@
 // Copyright Epic Games, Inc. All Rights Reserved.
 
-#pragma once
-
 #include "TakeRecorderAudioSettingsCustomization.h"
 #include "IDetailChildrenBuilder.h"
 #include "TakeRecorderSettings.h"
