Fix uninitialized variable.

diff --git a/Engine/Plugins/VirtualProduction/Takes/Source/TakesCore/Public/TakeRecorderSourceProperty.h b/Engine/Plugins/VirtualProduction/Takes/Source/TakesCore/Public/TakeRecorderSourceProperty.h
--- a/Engine/Plugins/VirtualProduction/Takes/Source/TakesCore/Public/TakeRecorderSourceProperty.h
+++ b/Engine/Plugins/VirtualProduction/Takes/Source/TakesCore/Public/TakeRecorderSourceProperty.h
@@ -162,7 +162,7 @@ struct TAKESCORE_API FAudioInputDeviceProperty
 
 	/** The number of input channels supported by this device */
 	UPROPERTY(Transient, EditAnywhere, BlueprintReadWrite, Category = "Property")
-	int32 DeviceInputChannelCount;
+	int32 DeviceInputChannelCount = 0;
 	
 	/** The desired buffer size used for audio callbacks during record */
 	UPROPERTY(config, EditAnywhere, BlueprintReadWrite, Category = "Property", meta = (ClampMin = "256", UIMin = "256", ClampMax = "8192", UIMax = "8192"))
@@ -182,7 +182,7 @@ public:
 
 	/** The currently selected channel from the details panel for this source */
 	UPROPERTY(EditAnywhere, BlueprintReadWrite, Category = "Property")
-	int32 AudioInputDeviceChannel;
+	int32 AudioInputDeviceChannel = 0;
 };
 
 DECLARE_MULTICAST_DELEGATE(FOnAudioInputDeviceChanged);
