Code cleanup for take recorder and cache track recorder

diff --git a/Engine/Plugins/VirtualProduction/Takes/Source/CacheTrackRecorder/CacheTrackRecorder.Build.cs b/Engine/Plugins/VirtualProduction/Takes/Source/CacheTrackRecorder/CacheTrackRecorder.Build.cs
--- a/Engine/Plugins/VirtualProduction/Takes/Source/CacheTrackRecorder/CacheTrackRecorder.Build.cs
+++ b/Engine/Plugins/VirtualProduction/Takes/Source/CacheTrackRecorder/CacheTrackRecorder.Build.cs
@@ -15,7 +15,6 @@ public class CacheTrackRecorder : ModuleRules
 				"Engine",
 				"LevelEditor",
 				"LevelSequence",
-				"LevelSequenceEditor",
 				"MovieScene",
 				"TakesCore",
 				"TakeMovieScene",
diff --git a/Engine/Plugins/VirtualProduction/Takes/Source/CacheTrackRecorder/Private/Recorder/CacheTrackRecorder.cpp b/Engine/Plugins/VirtualProduction/Takes/Source/CacheTrackRecorder/Private/Recorder/CacheTrackRecorder.cpp
--- a/Engine/Plugins/VirtualProduction/Takes/Source/CacheTrackRecorder/Private/Recorder/CacheTrackRecorder.cpp
+++ b/Engine/Plugins/VirtualProduction/Takes/Source/CacheTrackRecorder/Private/Recorder/CacheTrackRecorder.cpp
@@ -39,7 +39,6 @@
 
 // LevelEditor includes
 #include "IAssetViewport.h"
-#include "ILevelSequenceEditorToolkit.h"
 #include "LevelEditor.h"
 #include "Subsystems/AssetEditorSubsystem.h"
 
@@ -49,7 +48,6 @@
 
 FCacheRecorderUserParameters::FCacheRecorderUserParameters()
 	: bMaximizeViewport(false)
-	, CountdownSeconds(0.f)
 	, EngineTimeDilation(1.f)
 	, bResetPlayhead(true)
 	, bStopAtPlaybackEnd(true)
@@ -79,9 +77,6 @@ public:
 		WeakRecorder = InCacheTrackRecorder;
 		CacheTrackRecorderState = InCacheTrackRecorder->GetState();
 
-		UTakeMetaData* TakeMetaData = InCacheTrackRecorder->GetSequence()->FindMetaData<UTakeMetaData>();
-		check(TakeMetaData);
-
 		ChildSlot
 		[
 			SNew(SBox)
@@ -101,7 +96,7 @@ public:
 					[
 						SNew(STextBlock)
 						.Font(FCoreStyle::Get().GetFontStyle(TEXT("NotificationList.FontBold")))
-						.Text(FText::Format(LOCTEXT("RecordingTitleFormat", "Take {0} of slate {1}"), FText::AsNumber(TakeMetaData->GetTakeNumber()), FText::FromString(TakeMetaData->GetSlate())))
+						.Text(LOCTEXT("RecordingTitleFormat", "Recording sequencer cache"))
 					]
 
 					+ SHorizontalBox::Slot()
@@ -157,7 +152,7 @@ private:
 		{
 			ECacheTrackRecorderState NewCacheTrackRecorderState = Recorder->GetState();
 
-			if (NewCacheTrackRecorderState == ECacheTrackRecorderState::CountingDown || NewCacheTrackRecorderState == ECacheTrackRecorderState::Started)
+			if (NewCacheTrackRecorderState == ECacheTrackRecorderState::Starting || NewCacheTrackRecorderState == ECacheTrackRecorderState::Started)
 			{
 				// When counting down the text may change on tick
 				TextBlock->SetText(GetDetailText());
@@ -199,9 +194,9 @@ private:
 	{
 		if (UCacheTrackRecorder* Recorder = WeakRecorder.Get())
 		{
-			if (Recorder->GetState() == ECacheTrackRecorderState::CountingDown)
+			if (Recorder->GetState() == ECacheTrackRecorderState::Starting)
 			{
-				return FText::Format(LOCTEXT("CountdownText", "Recording in {0}s..."), FText::AsNumber(FMath::CeilToInt(Recorder->GetCountdownSeconds())));
+				return LOCTEXT("CountdownText", "Starting to record...");
 			}
 			if (Recorder->GetState() == ECacheTrackRecorderState::Stopped)
 			{
@@ -212,18 +207,11 @@ private:
 				return LOCTEXT("CancelledText", "Recording Cancelled");
 			}
 
-			ULevelSequence* LevelSequence = Recorder->GetSequence();
-			if (UTakeMetaData* TakeMetaData = LevelSequence ? LevelSequence->FindMetaData<UTakeMetaData>() : nullptr)
-			{
-				FFrameRate FrameRate = TakeMetaData->GetFrameRate();
-				FTimespan RecordingDuration = FDateTime::UtcNow() - TakeMetaData->GetTimestamp();
-
-				FFrameNumber TotalFrames = FFrameNumber(static_cast<int32>(FrameRate.AsDecimal() * RecordingDuration.GetTotalSeconds()));
-
-				FTimecode Timecode = FTimecode::FromFrameNumber(TotalFrames, FrameRate);
-
-				return FText::Format(LOCTEXT("RecordingTimecodeText", "Recording...{0}"), FText::FromString(Timecode.ToString()));
-			}
+			static const FNumberFormattingOptions SecondFormattingOptions = FNumberFormattingOptions()
+			.SetMinimumFractionalDigits(2)
+			.SetMaximumFractionalDigits(2);
+			FQualifiedFrameTime RecordTime = Recorder->GetRecordTime();
+			return FText::Format(LOCTEXT("RecordingTimecodeText", "Recording... {0}s"), FText::AsNumber(RecordTime.AsSeconds(), &SecondFormattingOptions));
 		}
 
 		return LOCTEXT("RecordingText", "Recording...");
@@ -239,8 +227,6 @@ private:
 	{
 	}
 
-private:
-
 	FReply ButtonClicked()
 	{
 		if (UCacheTrackRecorder* Recorder = WeakRecorder.Get())
@@ -250,11 +236,10 @@ private:
 		return FReply::Handled();
 	}
 
-private:
 	TSharedPtr<SWidget> Button, Throbber, Hyperlink;
 	TSharedPtr<STextBlock> TextBlock;
 
-	ECacheTrackRecorderState CacheTrackRecorderState = ECacheTrackRecorderState::CountingDown;
+	ECacheTrackRecorderState CacheTrackRecorderState = ECacheTrackRecorderState::Starting;
 	TWeakPtr<SNotificationItem> WeakOwningNotification;
 	TWeakObjectPtr<UCacheTrackRecorder> WeakRecorder;
 };
@@ -302,7 +287,6 @@ static TStrongObjectPtr<UCacheTrackRecorder>& GetCurrentRecorder()
 	static TStrongObjectPtr<UCacheTrackRecorder> CurrentRecorder;
 	return CurrentRecorder;
 }
-FOnCacheTrackRecordingInitialized UCacheTrackRecorder::OnRecordingInitializedEvent;
 
 // Static functions for UCacheTrackRecorder
 UCacheTrackRecorder* UCacheTrackRecorder::GetActiveRecorder()
@@ -310,11 +294,6 @@ UCacheTrackRecorder* UCacheTrackRecorder::GetActiveRecorder()
 	return GetCurrentRecorder().Get();
 }
 
-FOnCacheTrackRecordingInitialized& UCacheTrackRecorder::OnRecordingInitialized()
-{
-	return OnRecordingInitializedEvent;
-}
-
 void UCacheTrackRecorder::RecordCacheTrack(IMovieSceneCachedTrack* Track, TSharedPtr<ISequencer> Sequencer, FCacheRecorderParameters Parameters)
 {	
 	if (Track)
@@ -331,7 +310,6 @@ void UCacheTrackRecorder::RecordCacheTracks(const TArray<IMovieSceneCachedTrack*
 	if (LevelSequence && CacheTracks.Num() > 0)
 	{
 		Parameters.StartFrame = LevelSequence->GetMovieScene()->GetPlaybackRange().GetLowerBoundValue();
-		Parameters.User.CountdownSeconds = 0;
 		Parameters.Project.bStartAtCurrentTimecode = false;
 
 		// If not resetting the playhead, store the current time as the start frame for recording. 
@@ -349,7 +327,7 @@ void UCacheTrackRecorder::RecordCacheTracks(const TArray<IMovieSceneCachedTrack*
 		}
 
 		FText ErrorText = LOCTEXT("UnknownError", "An unknown error occurred when trying to start recording");
-		if (!NewRecorder->Initialize(LevelSequence, CacheTracks, TakeMetaData, Parameters, &ErrorText))
+		if (!NewRecorder->Initialize(CacheTracks, TakeMetaData->GetSlate(), Parameters, &ErrorText))
 		{
 			if (ensure(!ErrorText.IsEmpty()))
 			{
@@ -388,20 +366,12 @@ bool UCacheTrackRecorder::SetActiveRecorder(UCacheTrackRecorder* NewActiveRecord
 
 	GetCurrentRecorder().Reset(NewActiveRecorder);
 	TickableCacheTrackRecorder.WeakRecorder = GetCurrentRecorder().Get();
-	OnRecordingInitializedEvent.Broadcast(NewActiveRecorder);
 	return true;
 }
 
 // Non-static api for UCacheTrackRecorder
 
-UCacheTrackRecorder::UCacheTrackRecorder(const FObjectInitializer& ObjInit)
-	: Super(ObjInit)
-{
-	CountdownSeconds = 0.f;
-	SequenceAsset = nullptr;
-}
-
-bool UCacheTrackRecorder::Initialize(ULevelSequence* LevelSequenceBase, const TArray<IMovieSceneCachedTrack*>& InCacheTracks, const UTakeMetaData* MetaData, const FCacheRecorderParameters& InParameters, FText* OutError)
+bool UCacheTrackRecorder::Initialize(const TArray<IMovieSceneCachedTrack*>& InCacheTracks, const FString& Slate, const FCacheRecorderParameters& InParameters, FText* OutError)
 {
 	FGCObjectScopeGuard GCGuard(this);
 
@@ -414,7 +384,7 @@ bool UCacheTrackRecorder::Initialize(ULevelSequence* LevelSequenceBase, const TA
 		return false;
 	}
 
-	if (MetaData->GetSlate().IsEmpty())
+	if (Slate.IsEmpty())
 	{
 		if (OutError)
 		{
@@ -434,10 +404,9 @@ bool UCacheTrackRecorder::Initialize(ULevelSequence* LevelSequenceBase, const TA
 		CacheTracks.AddDefaulted_GetRef().Track = Track;
 	}
 
-	OnRecordingPreInitializeEvent.Broadcast(this);
-
 	FCacheRecorderParameters FinalParameters = InParameters;
-	if (!InitializeSequencer(LevelSequenceBase, OutError))
+	WeakSequencer = TakesUtils::OpenSequencer(SequenceAsset, OutError);
+	if (!WeakSequencer.IsValid())
 	{
 		return false;
 	}
@@ -453,7 +422,7 @@ bool UCacheTrackRecorder::Initialize(ULevelSequence* LevelSequenceBase, const TA
 	AddToRoot();
 
 	Parameters = FinalParameters;
-	State      = ECacheTrackRecorderState::CountingDown;
+	State      = ECacheTrackRecorderState::Starting;
 
 	// Perform any other parameter-configurable initialization. Must have a valid world at this point.
 	InitializeFromParameters();
@@ -500,25 +469,9 @@ bool UCacheTrackRecorder::Initialize(ULevelSequence* LevelSequenceBase, const TA
 
 void UCacheTrackRecorder::DiscoverSourceWorld()
 {
-	UWorld* WorldToRecordIn = nullptr;
+	WeakWorld = TakesUtils::DiscoverSourceWorld();
 
-	for (const FWorldContext& WorldContext : GEngine->GetWorldContexts())
-	{
-		if (WorldContext.WorldType == EWorldType::PIE || WorldContext.WorldType == EWorldType::Game )
-		{
-			WorldToRecordIn = WorldContext.World();
-			break;
-		}
-		else if (WorldContext.WorldType == EWorldType::Editor)
-		{
-			WorldToRecordIn = WorldContext.World();
-		}
-	}
-
-	check(WorldToRecordIn);
-	WeakWorld = WorldToRecordIn;
-
-	bool bPlayInGame = WorldToRecordIn->WorldType == EWorldType::PIE || WorldToRecordIn->WorldType == EWorldType::Game;
+	bool bPlayInGame = WeakWorld->WorldType == EWorldType::PIE || WeakWorld->WorldType == EWorldType::Game;
 	// If recording via PIE, be sure to stop recording cleanly when PIE ends
 	if ( bPlayInGame )
 	{
@@ -531,36 +484,8 @@ void UCacheTrackRecorder::DiscoverSourceWorld()
 	}
 }
 
-bool UCacheTrackRecorder::InitializeSequencer(ULevelSequence* LevelSequence, FText* OutError)
-{
-	if ( GEditor != nullptr )
-	{
-		// Open the sequence and set the sequencer ptr
-		GEditor->GetEditorSubsystem<UAssetEditorSubsystem>()->OpenEditorForAsset(LevelSequence);
-
-		IAssetEditorInstance*        AssetEditor         = GEditor->GetEditorSubsystem<UAssetEditorSubsystem>()->FindEditorForAsset(LevelSequence, false);
-		ILevelSequenceEditorToolkit* LevelSequenceEditor = static_cast<ILevelSequenceEditorToolkit*>(AssetEditor);
-
-		WeakSequencer = LevelSequenceEditor ? LevelSequenceEditor->GetSequencer() : nullptr;
-
-		if (!WeakSequencer.Pin().IsValid())
-		{
-			if (OutError)
-			{
-				*OutError = FText::Format(LOCTEXT("FailedToOpenSequencerError", "Failed to open Sequencer for asset '{0}."), FText::FromString(LevelSequence->GetPathName()));
-			}
-			return false;
-		}
-	}
-		
-	return true;
-}
-
 void UCacheTrackRecorder::InitializeFromParameters()
 {
-	// Initialize the countdown delay
-	CountdownSeconds = Parameters.User.CountdownSeconds;
-
 	// Set the end recording frame if enabled
 	StopRecordingFrame = Parameters.User.bStopAtPlaybackEnd ? SequenceAsset->GetMovieScene()->GetPlaybackRange().GetUpperBoundValue() : TOptional<FFrameNumber>();
 
@@ -642,14 +567,9 @@ UWorld* UCacheTrackRecorder::GetWorld() const
 
 void UCacheTrackRecorder::Tick(float DeltaTime)
 {
-	if (State == ECacheTrackRecorderState::CountingDown)
+	if (State == ECacheTrackRecorderState::Starting)
 	{
 		NumberOfTicksAfterPre = 0;
-		CountdownSeconds = FMath::Max(0.f, CountdownSeconds - DeltaTime);
-		if (CountdownSeconds > 0.f)
-		{
-			return;
-		}
 		PreRecord();
 	}
 	else if (State == ECacheTrackRecorderState::PreRecord)
@@ -673,36 +593,7 @@ void UCacheTrackRecorder::Tick(float DeltaTime)
 
 FQualifiedFrameTime UCacheTrackRecorder::GetRecordTime() const
 {
-	FQualifiedFrameTime RecordTime;
-
-	if (TSharedPtr<ISequencer> Sequencer = WeakSequencer.Pin(); Sequencer.IsValid())
-	{
-		RecordTime = Sequencer->GetGlobalTime();
-	}
-	else if (SequenceAsset)
-	{
-		if (UMovieScene* MovieScene = SequenceAsset->GetMovieScene())
-		{
-			FFrameRate FrameRate = MovieScene->GetDisplayRate();
-			FFrameRate TickResolution = MovieScene->GetTickResolution();
-
-			FTimecode CurrentTimecode = FApp::GetTimecode();
-		
-			FFrameNumber CurrentFrame = FFrameRate::TransformTime(FFrameTime(CurrentTimecode.ToFrameNumber(FrameRate)), FrameRate, TickResolution).FloorToFrame();
-			FFrameNumber FrameAtStart = FFrameRate::TransformTime(FFrameTime(TimecodeAtStart.ToFrameNumber(FrameRate)), FrameRate, TickResolution).FloorToFrame();
-
-			if (Parameters.Project.bStartAtCurrentTimecode)
-			{
-				RecordTime = FQualifiedFrameTime(CurrentFrame, TickResolution);
-			}
-			else
-			{
-				RecordTime = FQualifiedFrameTime(CurrentFrame - FrameAtStart, TickResolution);
-			}
-		}
-	}
-
-	return RecordTime;
+	return TakesUtils::GetRecordTime(WeakSequencer.Pin(), SequenceAsset, TimecodeAtStart, Parameters.Project.bStartAtCurrentTimecode);
 }
 
 void UCacheTrackRecorder::InternalTick(float DeltaTime)
@@ -778,7 +669,6 @@ void UCacheTrackRecorder::PreRecord()
 	if (AWorldSettings* WorldSettings = RecordingWorld->GetWorldSettings())
 	{
 		const float ExistingCinematicTimeDilation = WorldSettings->CinematicTimeDilation;
-
 		const bool bInvalidTimeDilation = Parameters.User.EngineTimeDilation == 0.f;
 
 		if (bInvalidTimeDilation)
@@ -902,8 +792,6 @@ void UCacheTrackRecorder::Start()
 		// Log in lieu of the notification widget
 		UE_LOG(LogTakesCore, Log, TEXT("Started recording"));
 	}
-
-	OnRecordingStartedEvent.Broadcast(this);
 }
 
 void UCacheTrackRecorder::Stop()
@@ -970,8 +858,6 @@ void UCacheTrackRecorder::StopInternal(const bool bCancelled)
 			}
 		}
 
-		OnRecordingStoppedEvent.Broadcast(this);
-
 		if (MovieScene)
 		{
 			TRange<FFrameNumber> Range = MovieScene->GetPlaybackRange();
@@ -1046,15 +932,6 @@ void UCacheTrackRecorder::StopInternal(const bool bCancelled)
 	{
 		GetCurrentRecorder().Reset();
 		TickableCacheTrackRecorder.WeakRecorder = nullptr;
-
-		if (bRecordingFinished)
-		{
-			OnRecordingFinishedEvent.Broadcast(this);
-		}
-		else
-		{
-			OnRecordingCancelledEvent.Broadcast(this);
-		}
 	}
 
 	if (SequenceAsset)
@@ -1069,34 +946,9 @@ void UCacheTrackRecorder::StopInternal(const bool bCancelled)
 	RemoveFromRoot();
 }
 
-FOnCacheTrackRecordingPreInitialize& UCacheTrackRecorder::OnRecordingPreInitialize()
-{
-	return OnRecordingPreInitializeEvent;
-}
-
-FOnCacheTrackRecordingStarted& UCacheTrackRecorder::OnRecordingStarted()
-{
-	return OnRecordingStartedEvent;
-}
-
-FOnCacheTrackRecordingStopped& UCacheTrackRecorder::OnRecordingStopped()
-{
-	return OnRecordingStoppedEvent;
-}
-
-FOnCacheTrackRecordingFinished& UCacheTrackRecorder::OnRecordingFinished()
-{
-	return OnRecordingFinishedEvent;
-}
-
-FOnCacheTrackRecordingCancelled& UCacheTrackRecorder::OnRecordingCancelled()
-{
-	return OnRecordingCancelledEvent;
-}
-
 void UCacheTrackRecorder::HandlePIE(bool bIsSimulating)
 {
-	ULevelSequence* FinishedAsset = GetSequence();
+	ULevelSequence* FinishedAsset = SequenceAsset;
 
 	if (ShouldShowNotifications())
 	{
diff --git a/Engine/Plugins/VirtualProduction/Takes/Source/CacheTrackRecorder/Public/Recorder/CacheTrackRecorder.h b/Engine/Plugins/VirtualProduction/Takes/Source/CacheTrackRecorder/Public/Recorder/CacheTrackRecorder.h
--- a/Engine/Plugins/VirtualProduction/Takes/Source/CacheTrackRecorder/Public/Recorder/CacheTrackRecorder.h
+++ b/Engine/Plugins/VirtualProduction/Takes/Source/CacheTrackRecorder/Public/Recorder/CacheTrackRecorder.h
@@ -9,7 +9,6 @@
 class ULevelSequence;
 class UMovieSceneTrackRecorder;
 class IMovieSceneCachedTrack;
-class UTakeMetaData;
 
 USTRUCT()
 struct CACHETRACKRECORDER_API FCacheRecorderUserParameters
@@ -22,10 +21,6 @@ struct CACHETRACKRECORDER_API FCacheRecorderUserParameters
 	UPROPERTY(config, EditAnywhere, Category="User Settings")
 	bool bMaximizeViewport;
 
-	/** Delay that we will use before starting recording */
-	UPROPERTY(config, EditAnywhere, Category="User Settings", DisplayName="Countdown", meta=(Units=s, ClampMin="0.0", UIMin="0.0", ClampMax="60.0", UIMax="60.0"))
-	float CountdownSeconds;
-
 	/** The engine time dilation to apply during the recording */
 	UPROPERTY(config, EditAnywhere, Category="User Settings", meta=(Units=Multiplier, ClampMin="0.00001", UIMin="0.00001"))
 	float EngineTimeDilation;
@@ -103,7 +98,7 @@ struct CACHETRACKRECORDER_API FCacheRecorderParameters
 UENUM(BlueprintType)
 enum class ECacheTrackRecorderState : uint8
 {
-	CountingDown,
+	Starting,
 	PreRecord,
 	TickingAfterPre,
 	Started,
@@ -111,13 +106,6 @@ enum class ECacheTrackRecorderState : uint8
 	Cancelled,
 };
 
-DECLARE_MULTICAST_DELEGATE_OneParam(FOnCacheTrackRecordingPreInitialize, UCacheTrackRecorder*);
-DECLARE_MULTICAST_DELEGATE_OneParam(FOnCacheTrackRecordingInitialized, UCacheTrackRecorder*);
-DECLARE_MULTICAST_DELEGATE_OneParam(FOnCacheTrackRecordingStarted, UCacheTrackRecorder*);
-DECLARE_MULTICAST_DELEGATE_OneParam(FOnCacheTrackRecordingStopped, UCacheTrackRecorder*);
-DECLARE_MULTICAST_DELEGATE_OneParam(FOnCacheTrackRecordingFinished, UCacheTrackRecorder*);
-DECLARE_MULTICAST_DELEGATE_OneParam(FOnCacheTrackRecordingCancelled, UCacheTrackRecorder*);
-
 USTRUCT()
 struct FCachedTrackSource
 {
@@ -136,36 +124,25 @@ public:
 
 	GENERATED_BODY()
 
-	UCacheTrackRecorder(const FObjectInitializer& ObjInit);
-
-public:
-
 	/**
 	 * Retrieve the currently active take recorder instance
 	 */
 	static UCacheTrackRecorder* GetActiveRecorder();
 
 	/**
-	 * Retrieve a multi-cast delegate that is triggered when a new recording begins
+	 * Utility method to record a single cache track
 	 */
-	static FOnCacheTrackRecordingInitialized& OnRecordingInitialized();
-
+	static void RecordCacheTrack(IMovieSceneCachedTrack* Track, TSharedPtr<ISequencer> Sequencer, FCacheRecorderParameters Parameters = FCacheRecorderParameters());
+	
 	/**
-	 * Utility method to quickly record a single cache track
+	 * Utility method to record several cache tracks at once
 	 */
-	static void RecordCacheTrack(IMovieSceneCachedTrack* Track, TSharedPtr<ISequencer> Sequencer, FCacheRecorderParameters Parameters = FCacheRecorderParameters());
 	static void RecordCacheTracks(const TArray<IMovieSceneCachedTrack*>& CacheTracks, TSharedPtr<ISequencer> Sequencer, FCacheRecorderParameters Parameters = FCacheRecorderParameters());
-	static void RecordSelectedTracks(TSharedPtr<ISequencer> Sequencer, FCacheRecorderParameters Parameters = FCacheRecorderParameters());
-public:
 
 	/**
-	 * Access the number of seconds remaining before this recording will start
+	 * Utility method to record all the cache tracks selected in sequencer
 	 */
-	UFUNCTION(BlueprintCallable, Category="Cache Recorder")
-	float GetCountdownSeconds() const
-	{
-		return CountdownSeconds;
-	}
+	static void RecordSelectedTracks(TSharedPtr<ISequencer> Sequencer, FCacheRecorderParameters Parameters = FCacheRecorderParameters());
 
 	/**
 	 * Access the sequence asset that this recorder is recording into
@@ -189,7 +166,7 @@ public:
 	 * Initialize a new recording with the specified parameters. Fails if another recording is currently in progress.
 	 *
 	 */
-	bool Initialize(ULevelSequence* LevelSequenceBase, const TArray<IMovieSceneCachedTrack*>& CacheTracks, const UTakeMetaData* MetaData, const FCacheRecorderParameters& InParameters, FText* OutError = nullptr);
+	bool Initialize(const TArray<IMovieSceneCachedTrack*>& CacheTracks, const FString& Slate, const FCacheRecorderParameters& InParameters, FText* OutError = nullptr);
 
 	/**
 	 * Called to stop the recording
@@ -201,30 +178,8 @@ public:
 	 */
 	void Cancel();
 
-	/**
-	 * Retrieve a multi-cast delegate that is triggered before initialization occurs (ie. when the recording button is pressed and before the countdown starts)
-	 */
-	FOnCacheTrackRecordingPreInitialize& OnRecordingPreInitialize();
-
-	/**
-	 * Retrieve a multi-cast delegate that is triggered when this recording starts
-	 */
-	FOnCacheTrackRecordingStarted& OnRecordingStarted();
-
-	/**
-	 * Retrieve a multi-cast delegate that is triggered when this recording is stopped
-	 */
-	FOnCacheTrackRecordingStopped& OnRecordingStopped();
-
-	/**
-	 * Retrieve a multi-cast delegate that is triggered when this recording finishes
-	 */
-	FOnCacheTrackRecordingFinished& OnRecordingFinished();
-
-	/**
-	 * Retrieve a multi-cast delegate that is triggered when this recording is cancelled
-	 */
-	FOnCacheTrackRecordingCancelled& OnRecordingCancelled();
+	/* The time at which to record. Taken from the Sequencer global time, otherwise based on timecode */
+	FQualifiedFrameTime GetRecordTime() const;
 
 private:
 
@@ -253,11 +208,6 @@ private:
 	 */
 	void HandlePIE(bool bIsSimulating);
 
-	/**
-	 * Attempt to open the sequencer UI for the asset to be recorded
-	 */
-	bool InitializeSequencer(ULevelSequence* LevelSequence, FText* OutError);
-
 	/**
 	 * Discovers the source world to record from, and initializes it for recording
 	 */
@@ -283,19 +233,11 @@ private:
 	/* Restores the editor tick state if necessary */
 	void RestoreEditorTickState();
 
-	/* The time at which to record. Taken from the Sequencer global time, otherwise based on timecode */
-	FQualifiedFrameTime GetRecordTime() const;
-
 	/** Called by Tick and Start to make sure we record at start */
 	void InternalTick(float DeltaTime);
 
 	virtual UWorld* GetWorld() const override;
 
-private:
-
-	/** The number of seconds remaining before Start() should be called */
-	float CountdownSeconds;
-
 	/** The state of this recorder instance */
 	ECacheTrackRecorderState State;
 
@@ -326,21 +268,6 @@ private:
 	/** Anonymous array of cleanup functions to perform when a recording has finished */
 	TArray<TFunction<void()>> OnStopCleanup;
 
-	/** Triggered before the recorder is initialized */
-	FOnCacheTrackRecordingPreInitialize OnRecordingPreInitializeEvent;
-
-	/** Triggered when this recorder starts */
-	FOnCacheTrackRecordingStarted OnRecordingStartedEvent;
-
-	/** Triggered when this recorder is stopped */
-	FOnCacheTrackRecordingStopped OnRecordingStoppedEvent;
-
-	/** Triggered when this recorder finishes */
-	FOnCacheTrackRecordingFinished OnRecordingFinishedEvent;
-
-	/** Triggered when this recorder is cancelled */
-	FOnCacheTrackRecordingCancelled OnRecordingCancelledEvent;
-
 	/** Sequencer ptr that controls playback of the desination asset during the recording */
 	TWeakPtr<ISequencer> WeakSequencer;
 
@@ -356,13 +283,7 @@ private:
 	 */
 	static bool SetActiveRecorder(UCacheTrackRecorder* NewActiveRecorder);
 
-	/** Event to trigger when a new recording is initialized */
-	static FOnCacheTrackRecordingInitialized OnRecordingInitializedEvent;
-
-	EAllowEditsMode CachedAllowEditsMode;
-	EAutoChangeMode CachedAutoChangeMode;
 	EUpdateClockSource CachedClockSource;
-	
 	TRange<FFrameNumber> CachedPlaybackRange;
 
 	struct FSavedState
diff --git a/Engine/Plugins/VirtualProduction/Takes/Source/TakeRecorder/Private/Recorder/TakeRecorder.cpp b/Engine/Plugins/VirtualProduction/Takes/Source/TakeRecorder/Private/Recorder/TakeRecorder.cpp
--- a/Engine/Plugins/VirtualProduction/Takes/Source/TakeRecorder/Private/Recorder/TakeRecorder.cpp
+++ b/Engine/Plugins/VirtualProduction/Takes/Source/TakeRecorder/Private/Recorder/TakeRecorder.cpp
@@ -25,16 +25,12 @@
 #include "Stats/Stats.h"
 #include "SequencerSettings.h"
 
-// LevelSequenceEditor includes
-#include "ILevelSequenceEditorToolkit.h"
-
 // Engine includes
 #include "GameFramework/WorldSettings.h"
 
 // UnrealEd includes
 #include "Editor.h"
 
-#include "ObjectTools.h"
 #include "UObject/GCObjectScopeGuard.h"
 
 // Slate includes
@@ -559,25 +555,9 @@ bool UTakeRecorder::Initialize ( ULevelSequence* LevelSequenceBase, UTakeRecorde
 
 void UTakeRecorder::DiscoverSourceWorld()
 {
-	UWorld* WorldToRecordIn = nullptr;
+	WeakWorld = TakesUtils::DiscoverSourceWorld();
 
-	for (const FWorldContext& WorldContext : GEngine->GetWorldContexts())
-	{
-		if (WorldContext.WorldType == EWorldType::PIE || WorldContext.WorldType == EWorldType::Game )
-		{
-			WorldToRecordIn = WorldContext.World();
-			break;
-		}
-		else if (WorldContext.WorldType == EWorldType::Editor)
-		{
-			WorldToRecordIn = WorldContext.World();
-		}
-	}
-
-	check(WorldToRecordIn);
-	WeakWorld = WorldToRecordIn;
-
-	bool bPlayInGame = WorldToRecordIn->WorldType == EWorldType::PIE || WorldToRecordIn->WorldType == EWorldType::Game;
+	bool bPlayInGame = WeakWorld->WorldType == EWorldType::PIE || WeakWorld->WorldType == EWorldType::Game;
 	// If recording via PIE, be sure to stop recording cleanly when PIE ends
 	if ( bPlayInGame )
 	{
@@ -589,7 +569,7 @@ void UTakeRecorder::DiscoverSourceWorld()
 			UClass* Class = StaticLoadClass(UTakeRecorderOverlayWidget::StaticClass(), nullptr, TEXT("/Takes/UMG/DefaultRecordingOverlay.DefaultRecordingOverlay_C"));
 			if (Class)
 			{
-				OverlayWidget = CreateWidget<UTakeRecorderOverlayWidget>(WorldToRecordIn, Class);
+				OverlayWidget = CreateWidget<UTakeRecorderOverlayWidget>(WeakWorld.Get(), Class);
 				OverlayWidget->SetFlags(RF_Transient);
 				OverlayWidget->SetRecorder(this);
 				OverlayWidget->AddToViewport();
@@ -640,19 +620,16 @@ bool UTakeRecorder::CreateDestinationAsset(const TCHAR* AssetPathFormat, ULevelS
 	SequenceAsset->MarkPackageDirty();
 	FAssetRegistryModule::AssetCreated(SequenceAsset);
 
-	if (!InitializeSequencer(SequenceAsset, OutError))
-	{
-		return false;
-	}
-
-	return true;
+	WeakSequencer = TakesUtils::OpenSequencer(SequenceAsset, OutError);
+	return WeakSequencer.IsValid();
 }
 
 bool UTakeRecorder::SetupDestinationAsset(const FTakeRecorderParameters& InParameters, ULevelSequence* LevelSequenceBase, UTakeRecorderSources* Sources, UTakeMetaData* MetaData, FText* OutError)
 {
 	check(LevelSequenceBase && Sources && MetaData);
 
-	if (!InitializeSequencer(LevelSequenceBase, OutError))
+	WeakSequencer = TakesUtils::OpenSequencer(SequenceAsset, OutError);
+	if (!WeakSequencer.IsValid())
 	{
 		return false;
 	}
@@ -684,31 +661,6 @@ bool UTakeRecorder::SetupDestinationAsset(const FTakeRecorderParameters& InParam
 	return true;
 }
 
-bool UTakeRecorder::InitializeSequencer(ULevelSequence* LevelSequence, FText* OutError)
-{
-	if ( GEditor != nullptr )
-	{
-		// Open the sequence and set the sequencer ptr
-		GEditor->GetEditorSubsystem<UAssetEditorSubsystem>()->OpenEditorForAsset(LevelSequence);
-
-		IAssetEditorInstance*        AssetEditor         = GEditor->GetEditorSubsystem<UAssetEditorSubsystem>()->FindEditorForAsset(LevelSequence, false);
-		ILevelSequenceEditorToolkit* LevelSequenceEditor = static_cast<ILevelSequenceEditorToolkit*>(AssetEditor);
-
-		WeakSequencer = LevelSequenceEditor ? LevelSequenceEditor->GetSequencer() : nullptr;
-
-		if (!WeakSequencer.Pin().IsValid())
-		{
-			if (OutError)
-			{
-				*OutError = FText::Format(LOCTEXT("FailedToOpenSequencerError", "Failed to open Sequencer for asset '{0}."), FText::FromString(LevelSequence->GetPathName()));
-			}
-			return false;
-		}
-	}
-		
-	return true;
-}
-
 void UTakeRecorder::InitializeFromParameters()
 {
 	// Initialize the countdown delay
@@ -789,38 +741,7 @@ void UTakeRecorder::Tick(float DeltaTime)
 
 FQualifiedFrameTime UTakeRecorder::GetRecordTime() const
 {
-	FQualifiedFrameTime RecordTime;
-
-	TSharedPtr<ISequencer> Sequencer = WeakSequencer.Pin();
-	if (Sequencer.IsValid())
-	{
-		RecordTime = Sequencer->GetGlobalTime();
-	}
-	else if (SequenceAsset)
-	{
-		UMovieScene* MovieScene = SequenceAsset->GetMovieScene();
-		if (MovieScene)
-		{
-			FFrameRate FrameRate = MovieScene->GetDisplayRate();
-			FFrameRate TickResolution = MovieScene->GetTickResolution();
-
-			FTimecode CurrentTimecode = FApp::GetTimecode();
-		
-			FFrameNumber CurrentFrame = FFrameRate::TransformTime(FFrameTime(CurrentTimecode.ToFrameNumber(FrameRate)), FrameRate, TickResolution).FloorToFrame();
-			FFrameNumber FrameAtStart = FFrameRate::TransformTime(FFrameTime(TimecodeAtStart.ToFrameNumber(FrameRate)), FrameRate, TickResolution).FloorToFrame();
-
-			if (Parameters.Project.bStartAtCurrentTimecode)
-			{
-				RecordTime = FQualifiedFrameTime(CurrentFrame, TickResolution);
-			}
-			else
-			{
-				RecordTime = FQualifiedFrameTime(CurrentFrame - FrameAtStart, TickResolution);
-			}
-		}
-	}
-
-	return RecordTime;
+	return TakesUtils::GetRecordTime(WeakSequencer.Pin(), SequenceAsset, TimecodeAtStart, Parameters.Project.bStartAtCurrentTimecode);
 }
 
 void UTakeRecorder::InternalTick(float DeltaTime)
diff --git a/Engine/Plugins/VirtualProduction/Takes/Source/TakeRecorder/Public/Recorder/TakeRecorder.h b/Engine/Plugins/VirtualProduction/Takes/Source/TakeRecorder/Public/Recorder/TakeRecorder.h
--- a/Engine/Plugins/VirtualProduction/Takes/Source/TakeRecorder/Public/Recorder/TakeRecorder.h
+++ b/Engine/Plugins/VirtualProduction/Takes/Source/TakeRecorder/Public/Recorder/TakeRecorder.h
@@ -203,11 +203,6 @@ private:
 	 */
 	bool SetupDestinationAsset(const FTakeRecorderParameters& InParameters, ULevelSequence* LevelSequenceBase, UTakeRecorderSources* Sources, UTakeMetaData* MetaData, FText* OutError);
 
-	/**
-	 * Attempt to open the sequencer UI for the asset to be recorded
-	 */
-	bool InitializeSequencer(ULevelSequence* LevelSequence, FText* OutError);
-
 	/**
 	 * Discovers the source world to record from, and initializes it for recording
 	 */
diff --git a/Engine/Plugins/VirtualProduction/Takes/Source/TakeRecorder/TakeRecorder.Build.cs b/Engine/Plugins/VirtualProduction/Takes/Source/TakeRecorder/TakeRecorder.Build.cs
--- a/Engine/Plugins/VirtualProduction/Takes/Source/TakeRecorder/TakeRecorder.Build.cs
+++ b/Engine/Plugins/VirtualProduction/Takes/Source/TakeRecorder/TakeRecorder.Build.cs
@@ -20,7 +20,6 @@ public class TakeRecorder : ModuleRules
 				"InputCore",
 				"LevelEditor",
 				"LevelSequence",
-				"LevelSequenceEditor",
 				"MovieScene",
 				"MovieSceneTools",
 				"PropertyEditor",
diff --git a/Engine/Plugins/VirtualProduction/Takes/Source/TakesCore/Private/TakesUtils.cpp b/Engine/Plugins/VirtualProduction/Takes/Source/TakesCore/Private/TakesUtils.cpp
--- a/Engine/Plugins/VirtualProduction/Takes/Source/TakesCore/Private/TakesUtils.cpp
+++ b/Engine/Plugins/VirtualProduction/Takes/Source/TakesCore/Private/TakesUtils.cpp
@@ -7,6 +7,7 @@
 #include "Engine/Engine.h"
 #include "Engine/World.h"
 #include "LevelSequence.h"
+#include "ILevelSequenceEditorToolkit.h"
 #include "MovieScene.h"
 #include "Math/Range.h"
 #include "MovieSceneSection.h"
@@ -16,6 +17,8 @@
 #include "Sections/MovieSceneCameraCutSection.h"
 #include "UObject/SavePackage.h"
 
+#define LOCTEXT_NAMESPACE "TakesUtils"
+
 namespace TakesUtils
 {
 
@@ -143,4 +146,82 @@ void CreateCameraCutTrack(ULevelSequence* LevelSequence, const FGuid& RecordedCa
 	CameraCutTrack->AddSection(*CameraCutSection);
 }
 
+UWorld* DiscoverSourceWorld()
+{
+	UWorld* SourceWorld = nullptr;
+
+	for (const FWorldContext& WorldContext : GEngine->GetWorldContexts())
+	{
+		if (WorldContext.WorldType == EWorldType::PIE || WorldContext.WorldType == EWorldType::Game )
+		{
+			SourceWorld = WorldContext.World();
+			break;
+		}
+		if (WorldContext.WorldType == EWorldType::Editor)
+		{
+			SourceWorld = WorldContext.World();
+		}
+	}
+
+	check(SourceWorld);
+	return SourceWorld;
+}
+
+TSharedPtr<ISequencer> OpenSequencer(ULevelSequence* LevelSequence, FText* OutError)
+{
+	TSharedPtr<ISequencer> Sequencer;
+	if ( GEditor != nullptr )
+	{
+		// Open the sequence and set the sequencer ptr
+		GEditor->GetEditorSubsystem<UAssetEditorSubsystem>()->OpenEditorForAsset(LevelSequence);
+
+		IAssetEditorInstance*        AssetEditor         = GEditor->GetEditorSubsystem<UAssetEditorSubsystem>()->FindEditorForAsset(LevelSequence, false);
+		ILevelSequenceEditorToolkit* LevelSequenceEditor = static_cast<ILevelSequenceEditorToolkit*>(AssetEditor);
+
+		Sequencer = LevelSequenceEditor ? LevelSequenceEditor->GetSequencer() : nullptr;
+
+		if (!Sequencer.IsValid() && OutError)
+		{
+			*OutError = FText::Format(LOCTEXT("FailedToOpenSequencerError", "Failed to open Sequencer for asset '{0}."), FText::FromString(LevelSequence->GetPathName()));
+		}
+	}
+		
+	return Sequencer;
+}
+
+FQualifiedFrameTime GetRecordTime(TSharedPtr<ISequencer> Sequencer, ULevelSequence* SequenceAsset, const FTimecode& TimecodeAtStart, bool bStartAtCurrentTimecode)
+{
+	FQualifiedFrameTime RecordTime;
+
+	if (Sequencer.IsValid())
+	{
+		RecordTime = Sequencer->GetGlobalTime();
+	}
+	else if (SequenceAsset)
+	{
+		if (UMovieScene* MovieScene = SequenceAsset->GetMovieScene())
+		{
+			FFrameRate FrameRate = MovieScene->GetDisplayRate();
+			FFrameRate TickResolution = MovieScene->GetTickResolution();
+
+			FTimecode CurrentTimecode = FApp::GetTimecode();
+		
+			FFrameNumber CurrentFrame = FFrameRate::TransformTime(FFrameTime(CurrentTimecode.ToFrameNumber(FrameRate)), FrameRate, TickResolution).FloorToFrame();
+			FFrameNumber FrameAtStart = FFrameRate::TransformTime(FFrameTime(TimecodeAtStart.ToFrameNumber(FrameRate)), FrameRate, TickResolution).FloorToFrame();
+
+			if (bStartAtCurrentTimecode)
+			{
+				RecordTime = FQualifiedFrameTime(CurrentFrame, TickResolution);
+			}
+			else
+			{
+				RecordTime = FQualifiedFrameTime(CurrentFrame - FrameAtStart, TickResolution);
+			}
+		}
+	}
+
+	return RecordTime;
 }
+}
+
+#undef LOCTEXT_NAMESPACE
diff --git a/Engine/Plugins/VirtualProduction/Takes/Source/TakesCore/Public/TakesUtils.h b/Engine/Plugins/VirtualProduction/Takes/Source/TakesCore/Public/TakesUtils.h
--- a/Engine/Plugins/VirtualProduction/Takes/Source/TakesCore/Public/TakesUtils.h
+++ b/Engine/Plugins/VirtualProduction/Takes/Source/TakesCore/Public/TakesUtils.h
@@ -4,6 +4,7 @@
 
 #include "AssetRegistry/AssetData.h"
 #include "CoreMinimal.h"
+#include "ISequencer.h"
 #include "Misc/PackageName.h"
 #include "MovieScene.h"
 #include "Modules/ModuleManager.h"
@@ -15,7 +16,6 @@
 
 class UWorld;
 class ULevelSequence;
-class UMovieScene;
 
 namespace TakesUtils
 {
@@ -157,4 +157,11 @@ namespace TakesUtils
 
 		return nullptr;
 	}
+
+	TAKESCORE_API UWorld* DiscoverSourceWorld();
+
+	TAKESCORE_API TSharedPtr<ISequencer> OpenSequencer(ULevelSequence* LevelSequence, FText* OutError);
+	
+	/* The time at which to record. Taken from the Sequencer global time, otherwise based on timecode */
+	TAKESCORE_API FQualifiedFrameTime GetRecordTime(TSharedPtr<ISequencer> Sequencer, ULevelSequence* SequenceAsset, const FTimecode& TimecodeAtStart, bool bStartAtCurrentTimecode);
 }
\ No newline at end of file
diff --git a/Engine/Plugins/VirtualProduction/Takes/Source/TakesCore/TakesCore.Build.cs b/Engine/Plugins/VirtualProduction/Takes/Source/TakesCore/TakesCore.Build.cs
--- a/Engine/Plugins/VirtualProduction/Takes/Source/TakesCore/TakesCore.Build.cs
+++ b/Engine/Plugins/VirtualProduction/Takes/Source/TakesCore/TakesCore.Build.cs
@@ -20,14 +20,15 @@ public class TakesCore : ModuleRules
 				"AssetRegistry",
 				"LevelSequence",
 				"SlateCore",
-                "MovieScene",
-                "MovieSceneTracks",
-                "MovieSceneTools",
+				"MovieScene",
+				"MovieSceneTracks",
+				"MovieSceneTools",
 				"LevelSequence",
+				"LevelSequenceEditor",
 				"Engine",
-                "SerializedRecorderInterface",
+				"SerializedRecorderInterface",
 
-            }
-        );
+			}
+		);
 	}
 }
