Fixed a bug where the wrong asset was used to initialize sequencer recording

diff --git a/Engine/Plugins/VirtualProduction/Takes/Source/CacheTrackRecorder/Private/Recorder/CacheTrackRecorder.cpp b/Engine/Plugins/VirtualProduction/Takes/Source/CacheTrackRecorder/Private/Recorder/CacheTrackRecorder.cpp
--- a/Engine/Plugins/VirtualProduction/Takes/Source/CacheTrackRecorder/Private/Recorder/CacheTrackRecorder.cpp
+++ b/Engine/Plugins/VirtualProduction/Takes/Source/CacheTrackRecorder/Private/Recorder/CacheTrackRecorder.cpp
@@ -327,7 +327,7 @@ void UCacheTrackRecorder::RecordCacheTracks(const TArray<IMovieSceneCachedTrack*
 		}
 
 		FText ErrorText = LOCTEXT("UnknownError", "An unknown error occurred when trying to start recording");
-		if (!NewRecorder->Initialize(CacheTracks, TakeMetaData->GetSlate(), Parameters, &ErrorText))
+		if (!NewRecorder->Initialize(LevelSequence, CacheTracks, TakeMetaData->GetSlate(), Parameters, &ErrorText))
 		{
 			if (ensure(!ErrorText.IsEmpty()))
 			{
@@ -371,7 +371,7 @@ bool UCacheTrackRecorder::SetActiveRecorder(UCacheTrackRecorder* NewActiveRecord
 
 // Non-static api for UCacheTrackRecorder
 
-bool UCacheTrackRecorder::Initialize(const TArray<IMovieSceneCachedTrack*>& InCacheTracks, const FString& Slate, const FCacheRecorderParameters& InParameters, FText* OutError)
+bool UCacheTrackRecorder::Initialize(ULevelSequence* LevelSequence, const TArray<IMovieSceneCachedTrack*>& InCacheTracks, const FString& Slate, const FCacheRecorderParameters& InParameters, FText* OutError)
 {
 	FGCObjectScopeGuard GCGuard(this);
 
@@ -405,7 +405,7 @@ bool UCacheTrackRecorder::Initialize(const TArray<IMovieSceneCachedTrack*>& InCa
 	}
 
 	FCacheRecorderParameters FinalParameters = InParameters;
-	WeakSequencer = TakesUtils::OpenSequencer(SequenceAsset, OutError);
+	WeakSequencer = TakesUtils::OpenSequencer(LevelSequence, OutError);
 	if (!WeakSequencer.IsValid())
 	{
 		return false;
diff --git a/Engine/Plugins/VirtualProduction/Takes/Source/CacheTrackRecorder/Public/Recorder/CacheTrackRecorder.h b/Engine/Plugins/VirtualProduction/Takes/Source/CacheTrackRecorder/Public/Recorder/CacheTrackRecorder.h
--- a/Engine/Plugins/VirtualProduction/Takes/Source/CacheTrackRecorder/Public/Recorder/CacheTrackRecorder.h
+++ b/Engine/Plugins/VirtualProduction/Takes/Source/CacheTrackRecorder/Public/Recorder/CacheTrackRecorder.h
@@ -166,7 +166,7 @@ public:
 	 * Initialize a new recording with the specified parameters. Fails if another recording is currently in progress.
 	 *
 	 */
-	bool Initialize(const TArray<IMovieSceneCachedTrack*>& CacheTracks, const FString& Slate, const FCacheRecorderParameters& InParameters, FText* OutError = nullptr);
+	bool Initialize(ULevelSequence* LevelSequence, const TArray<IMovieSceneCachedTrack*>& CacheTracks, const FString& Slate, const FCacheRecorderParameters& InParameters, FText* OutError = nullptr);
 
 	/**
 	 * Called to stop the recording
diff --git a/Engine/Plugins/VirtualProduction/Takes/Source/TakeRecorder/Private/Recorder/TakeRecorder.cpp b/Engine/Plugins/VirtualProduction/Takes/Source/TakeRecorder/Private/Recorder/TakeRecorder.cpp
--- a/Engine/Plugins/VirtualProduction/Takes/Source/TakeRecorder/Private/Recorder/TakeRecorder.cpp
+++ b/Engine/Plugins/VirtualProduction/Takes/Source/TakeRecorder/Private/Recorder/TakeRecorder.cpp
@@ -628,7 +628,7 @@ bool UTakeRecorder::SetupDestinationAsset(const FTakeRecorderParameters& InParam
 {
 	check(LevelSequenceBase && Sources && MetaData);
 
-	WeakSequencer = TakesUtils::OpenSequencer(SequenceAsset, OutError);
+	WeakSequencer = TakesUtils::OpenSequencer(LevelSequenceBase, OutError);
 	if (!WeakSequencer.IsValid())
 	{
 		return false;
