Fixed mismatched copy constructor/assignment operator in FAudioInputDeviceInfoProperty.

diff --git a/Engine/Plugins/VirtualProduction/Takes/Source/TakesCore/Public/TakeRecorderSourceProperty.h b/Engine/Plugins/VirtualProduction/Takes/Source/TakesCore/Public/TakeRecorderSourceProperty.h
--- a/Engine/Plugins/VirtualProduction/Takes/Source/TakesCore/Public/TakeRecorderSourceProperty.h
+++ b/Engine/Plugins/VirtualProduction/Takes/Source/TakesCore/Public/TakeRecorderSourceProperty.h
@@ -104,7 +104,6 @@ struct TAKESCORE_API FAudioInputDeviceInfoProperty
 	GENERATED_BODY()
 
 	FAudioInputDeviceInfoProperty() = default;
-	FAudioInputDeviceInfoProperty(const FAudioInputDeviceInfoProperty& InObject) = default;
 
 	FAudioInputDeviceInfoProperty(const FString& InDeviceName, const FString& InDeviceId, int32 InChannelCount, int32 InSampleRate, bool bInIsDefaultDevice)
 		: DeviceName(InDeviceName)
