Take Recorder: Bind recording commands to play world so that they work in editor and in PIE.

diff --git a/Engine/Plugins/VirtualProduction/Takes/Source/TakeRecorder/Private/Widgets/STakeRecorderCockpit.cpp b/Engine/Plugins/VirtualProduction/Takes/Source/TakeRecorder/Private/Widgets/STakeRecorderCockpit.cpp
--- a/Engine/Plugins/VirtualProduction/Takes/Source/TakeRecorder/Private/Widgets/STakeRecorderCockpit.cpp
+++ b/Engine/Plugins/VirtualProduction/Takes/Source/TakeRecorder/Private/Widgets/STakeRecorderCockpit.cpp
@@ -60,6 +60,7 @@
 #include "LevelEditor.h"
 #include "ISettingsModule.h"
 #include "Dialog/SMessageDialog.h"
+#include "Kismet2/DebuggerCommands.h"
 #include "Subsystems/AssetEditorSubsystem.h"
 
 #include "ISequencer.h"
@@ -163,8 +164,6 @@ void STakeRecorderCockpit::Construct(const FArguments& InArgs)
 	UpdateTakeError();
 	UpdateRecordError();
 
-	CommandList = MakeShareable(new FUICommandList);
-	
 	DigitsTypeInterface = MakeShareable(new FDigitsTypeInterface);
 
 	BindCommands();
@@ -1196,15 +1195,13 @@ void STakeRecorderCockpit::OnRecordingFinished(UTakeRecorder* Recorder)
 
 void STakeRecorderCockpit::BindCommands()
 {
-	CommandList->MapAction(FTakeRecorderCommands::Get().StartRecording,
+	// Bind our commands to the play world so that we can record in editor and in PIE
+	FPlayWorldCommands::GlobalPlayWorldActions->MapAction(
+		FTakeRecorderCommands::Get().StartRecording,
 		FExecuteAction::CreateSP(this, &STakeRecorderCockpit::StartRecording));
-
-	CommandList->MapAction(FTakeRecorderCommands::Get().StopRecording,
+	FPlayWorldCommands::GlobalPlayWorldActions->MapAction(
+		FTakeRecorderCommands::Get().StopRecording,
 		FExecuteAction::CreateSP(this, &STakeRecorderCockpit::StopRecording));
-
-	// Append to level editor module so that shortcuts are accessible in level editor
-	FLevelEditorModule& LevelEditorModule = FModuleManager::GetModuleChecked<FLevelEditorModule>(TEXT("LevelEditor"));
-	LevelEditorModule.GetGlobalLevelEditorActions()->Append(CommandList.ToSharedRef());
 }
 
 void STakeRecorderCockpit::OnToggleEditPreviousRecording(ECheckBoxState CheckState)
diff --git a/Engine/Plugins/VirtualProduction/Takes/Source/TakeRecorder/Private/Widgets/STakeRecorderCockpit.h b/Engine/Plugins/VirtualProduction/Takes/Source/TakeRecorder/Private/Widgets/STakeRecorderCockpit.h
--- a/Engine/Plugins/VirtualProduction/Takes/Source/TakeRecorder/Private/Widgets/STakeRecorderCockpit.h
+++ b/Engine/Plugins/VirtualProduction/Takes/Source/TakeRecorder/Private/Widgets/STakeRecorderCockpit.h
@@ -14,7 +14,6 @@ struct FFrameRate;
 class ULevelSequence;
 class UTakeMetaData;
 class UTakeRecorder;
-class FUICommandList;
 struct FDigitsTypeInterface;
 
 /**
@@ -178,8 +177,6 @@ private:
 	FDelegateHandle OnAssetRegistryFilesLoadedHandle;
 	FDelegateHandle OnRecordingInitializedHandle, OnRecordingFinishedHandle;
 
-	TSharedPtr<FUICommandList> CommandList;
-
 	TSharedPtr<FDigitsTypeInterface> DigitsTypeInterface;
 
 	// Cached take numbers and slate used to UpdateTakeError() only when necesary
