Check channel index upper bound when audio device changes.

diff --git a/Engine/Plugins/VirtualProduction/Takes/Source/TakeRecorderSources/Private/TakeRecorderMicrophoneAudioSource.cpp b/Engine/Plugins/VirtualProduction/Takes/Source/TakeRecorderSources/Private/TakeRecorderMicrophoneAudioSource.cpp
--- a/Engine/Plugins/VirtualProduction/Takes/Source/TakeRecorderSources/Private/TakeRecorderMicrophoneAudioSource.cpp
+++ b/Engine/Plugins/VirtualProduction/Takes/Source/TakeRecorderSources/Private/TakeRecorderMicrophoneAudioSource.cpp
@@ -92,6 +92,8 @@ void UTakeRecorderMicrophoneAudioSource::Initialize()
 		{
 			SetCurrentInputChannel(AvailableChannelNumber);
 		}
+
+		AudioInputManager->GetOnNotifySourcesOfDeviceChange().AddUObject(this, &UTakeRecorderMicrophoneAudioSource::OnNotifySourcesOfDeviceChange);
 	}
 }
 
@@ -122,7 +124,7 @@ TUniquePtr<TArray<bool>> UTakeRecorderMicrophoneAudioSource::GetChannelsInUse(co
 		{
 			int32 MicSourceChannelNumber = MicSource->AudioChannel.AudioInputDeviceChannel;
 			int32 ChannelIndex = MicSourceChannelNumber - 1;
-			if (ChannelIndex >= 0)
+			if (ChannelIndex >= 0 && ChannelIndex < InDeviceChannelCount)
 			{
 				(*ChannelsInUse)[ChannelIndex] = true;
 			}
@@ -234,7 +236,6 @@ void UTakeRecorderMicrophoneAudioSource::StartRecording(const FTimecode& InSecti
 			}
 		}
 
-
 		AudioInputManager->StartRecording(LastChannelInUse);
 	}
 }
