Take Recorder: Make setting DeviceChannelCountHandle not transactable

diff --git a/Engine/Plugins/VirtualProduction/Takes/Source/TakeRecorder/Private/Widgets/TakeRecorderAudioSettingsCustomization.cpp b/Engine/Plugins/VirtualProduction/Takes/Source/TakeRecorder/Private/Widgets/TakeRecorderAudioSettingsCustomization.cpp
--- a/Engine/Plugins/VirtualProduction/Takes/Source/TakeRecorder/Private/Widgets/TakeRecorderAudioSettingsCustomization.cpp
+++ b/Engine/Plugins/VirtualProduction/Takes/Source/TakeRecorder/Private/Widgets/TakeRecorderAudioSettingsCustomization.cpp
@@ -166,14 +166,14 @@ TSharedRef<SWidget> FAudioInputDevicePropertyCustomization::MakeAudioInputSelect
 				CurrentDeviceId = DeviceId;
 				InputDeviceHandle->SetValue(CurrentDeviceId);
 				SelectedAudioInputDevice = DeviceInfoHandle;
-				DeviceChannelCountHandle->SetValue(GetChannelCountFromInfoProperty(DeviceInfoHandle));
+				DeviceChannelCountHandle->SetValue(GetChannelCountFromInfoProperty(DeviceInfoHandle), EPropertyValueSetFlags::NotTransactable);
 			}
 			else if (bHaveCurrentDevice)
 			{
 				if (CurrentDeviceId == DeviceId)
 				{
 					SelectedAudioInputDevice = DeviceInfoHandle;
-					DeviceChannelCountHandle->SetValue(GetChannelCountFromInfoProperty(DeviceInfoHandle));
+					DeviceChannelCountHandle->SetValue(GetChannelCountFromInfoProperty(DeviceInfoHandle), EPropertyValueSetFlags::NotTransactable);
 				}
 			}
 
@@ -184,6 +184,7 @@ TSharedRef<SWidget> FAudioInputDevicePropertyCustomization::MakeAudioInputSelect
 	// Combo box component:
 	TSharedRef<SComboBox<TSharedPtr<IPropertyHandle>>> NewWidget = SNew(SComboBox<TSharedPtr<IPropertyHandle>>)
 	.OptionsSource(&AudioInputDevices)
+	.InitiallySelectedItem(SelectedAudioInputDevice)
 	.OnGenerateWidget_Lambda([](TSharedPtr<IPropertyHandle> InItem)
 	{
 		return SNew(STextBlock)
@@ -194,7 +195,11 @@ TSharedRef<SWidget> FAudioInputDevicePropertyCustomization::MakeAudioInputSelect
 	{
 		if (InSelection.IsValid() && ComboBoxTitleBlock.IsValid())
 		{
-			DeviceChannelCountHandle->SetValue(GetChannelCountFromInfoProperty(InSelection));
+			EPropertyValueSetFlags::Type SetValueFlags = (InSelectInfo == ESelectInfo::Direct) ?
+				EPropertyValueSetFlags::NotTransactable :
+				EPropertyValueSetFlags::DefaultFlags;
+			
+			DeviceChannelCountHandle->SetValue(GetChannelCountFromInfoProperty(InSelection), SetValueFlags);
 			InputDeviceHandle->SetValue(GetDeviceIdFromInfoProperty(InSelection));
 			ComboBoxTitleBlock->SetText(FText::AsCultureInvariant(GetDeviceNameFromInfoProperty(InSelection)));
 		}
