Small refactor of TakeRecorderMicrophoneAudioManager (removing unnecessary property).

diff --git a/Engine/Plugins/VirtualProduction/Takes/Source/TakeRecorder/Private/Widgets/TakeRecorderAudioSettingsCustomization.cpp b/Engine/Plugins/VirtualProduction/Takes/Source/TakeRecorder/Private/Widgets/TakeRecorderAudioSettingsCustomization.cpp
--- a/Engine/Plugins/VirtualProduction/Takes/Source/TakeRecorder/Private/Widgets/TakeRecorderAudioSettingsCustomization.cpp
+++ b/Engine/Plugins/VirtualProduction/Takes/Source/TakeRecorder/Private/Widgets/TakeRecorderAudioSettingsCustomization.cpp
@@ -11,7 +11,6 @@ void FAudioInputDevicePropertyCustomization::CustomizeChildren(TSharedRef<IPrope
 {
 	UseSystemDefaultHandle = PropertyHandle->GetChildHandle(GET_MEMBER_NAME_CHECKED(FAudioInputDeviceProperty, bUseSystemDefaultAudioDevice));
 	InputDeviceHandle = PropertyHandle->GetChildHandle(GET_MEMBER_NAME_CHECKED(FAudioInputDeviceProperty, DeviceId));
-	DeviceChannelCountHandle = PropertyHandle->GetChildHandle(GET_MEMBER_NAME_CHECKED(FAudioInputDeviceProperty, DeviceInputChannelCount));
 	BufferSizeHandle = PropertyHandle->GetChildHandle(GET_MEMBER_NAME_CHECKED(FAudioInputDeviceProperty, AudioInputBufferSize));
 	DeviceInfoArrayHandle = PropertyHandle->GetChildHandle(GET_MEMBER_NAME_CHECKED(FAudioInputDeviceProperty, DeviceInfoArray))->AsArray();
 
@@ -44,7 +43,7 @@ void FAudioInputDevicePropertyCustomization::CustomizeChildren(TSharedRef<IPrope
 	}
 
 	// Create audio input device combobox
-	if (InputDeviceHandle.IsValid() && DeviceChannelCountHandle.IsValid() && DeviceInfoArrayHandle.IsValid())
+	if (InputDeviceHandle.IsValid() && DeviceInfoArrayHandle.IsValid())
 	{
 		IDetailPropertyRow& InputDevicePropertyRow = ChildBuilder.AddProperty(PropertyHandle);
 
@@ -166,14 +165,12 @@ TSharedRef<SWidget> FAudioInputDevicePropertyCustomization::MakeAudioInputSelect
 				CurrentDeviceId = DeviceId;
 				InputDeviceHandle->SetValue(CurrentDeviceId);
 				SelectedAudioInputDevice = DeviceInfoHandle;
-				DeviceChannelCountHandle->SetValue(GetChannelCountFromInfoProperty(DeviceInfoHandle), EPropertyValueSetFlags::NotTransactable);
 			}
 			else if (bHaveCurrentDevice)
 			{
 				if (CurrentDeviceId == DeviceId)
 				{
 					SelectedAudioInputDevice = DeviceInfoHandle;
-					DeviceChannelCountHandle->SetValue(GetChannelCountFromInfoProperty(DeviceInfoHandle), EPropertyValueSetFlags::NotTransactable);
 				}
 			}
 
@@ -195,11 +192,6 @@ TSharedRef<SWidget> FAudioInputDevicePropertyCustomization::MakeAudioInputSelect
 	{
 		if (InSelection.IsValid() && ComboBoxTitleBlock.IsValid())
 		{
-			EPropertyValueSetFlags::Type SetValueFlags = (InSelectInfo == ESelectInfo::Direct) ?
-				EPropertyValueSetFlags::NotTransactable :
-				EPropertyValueSetFlags::DefaultFlags;
-			
-			DeviceChannelCountHandle->SetValue(GetChannelCountFromInfoProperty(InSelection), SetValueFlags);
 			InputDeviceHandle->SetValue(GetDeviceIdFromInfoProperty(InSelection));
 			ComboBoxTitleBlock->SetText(FText::AsCultureInvariant(GetDeviceNameFromInfoProperty(InSelection)));
 		}
diff --git a/Engine/Plugins/VirtualProduction/Takes/Source/TakeRecorderSources/Private/TakeRecorderMicrophoneAudioManager.cpp b/Engine/Plugins/VirtualProduction/Takes/Source/TakeRecorderSources/Private/TakeRecorderMicrophoneAudioManager.cpp
--- a/Engine/Plugins/VirtualProduction/Takes/Source/TakeRecorderSources/Private/TakeRecorderMicrophoneAudioManager.cpp
+++ b/Engine/Plugins/VirtualProduction/Takes/Source/TakeRecorderSources/Private/TakeRecorderMicrophoneAudioManager.cpp
@@ -33,7 +33,7 @@ void UTakeRecorderMicrophoneAudioManager::PostEditChangeProperty(FPropertyChange
 	}
 
 	// Notify the Mic sources of the new device channel count
-	GetOnNotifySourcesOfDeviceChange().Broadcast(AudioInputDevice.DeviceInputChannelCount);
+	GetOnNotifySourcesOfDeviceChange().Broadcast(GetDeviceChannelCount());
 	// Notify the UI of the change so it can update if needed
 	GetOnAudioInputDeviceChanged().Broadcast();
 }
@@ -111,6 +111,22 @@ void UTakeRecorderMicrophoneAudioManager::FinalizeRecording()
 	}
 }
 
+int32 UTakeRecorderMicrophoneAudioManager::GetDeviceChannelCount()
+{
+	if (AudioInputDevice.DeviceId.Len() > 0)
+	{
+		for (FAudioInputDeviceInfoProperty DeviceInfo : AudioInputDevice.DeviceInfoArray)
+		{
+			if (DeviceInfo.DeviceId == AudioInputDevice.DeviceId)
+			{
+				return DeviceInfo.InputChannels;
+			}
+		}
+	}
+
+	return 0;
+}
+
 void UTakeRecorderMicrophoneAudioManager::BuildDeviceInfoArray()
 {
 	AudioInputDevice.DeviceInfoArray.Empty();
diff --git a/Engine/Plugins/VirtualProduction/Takes/Source/TakeRecorderSources/Private/TakeRecorderMicrophoneAudioManager.h b/Engine/Plugins/VirtualProduction/Takes/Source/TakeRecorderSources/Private/TakeRecorderMicrophoneAudioManager.h
--- a/Engine/Plugins/VirtualProduction/Takes/Source/TakeRecorderSources/Private/TakeRecorderMicrophoneAudioManager.h
+++ b/Engine/Plugins/VirtualProduction/Takes/Source/TakeRecorderSources/Private/TakeRecorderMicrophoneAudioManager.h
@@ -33,7 +33,7 @@ public:
 	virtual void EnumerateAudioDevices(bool InForceRefresh = false) override;
 
 	/** Returns input channel count for currently selected audio device */
-	virtual int32 GetDeviceChannelCount() override { return AudioInputDevice.DeviceInputChannelCount;  }
+	virtual int32 GetDeviceChannelCount() override;
 
 	/** The audio device to use for this microphone source */
 	UPROPERTY(config, EditAnywhere, BlueprintReadWrite, Category = "Source", meta = (ShowOnlyInnerProperties))
@@ -58,8 +58,6 @@ public:
 
 	/** Fetches the USoundWave for this source after a Take has been recorded */
 	TObjectPtr<class USoundWave> GetRecordedSoundWave(const FTakeRecorderAudioSourceSettings& InSourceSettings);
-	/** Gets the cached value for the audio input device channel count */
-	int32 GetDeviceChannelCount() const { return AudioInputDevice.DeviceInputChannelCount; }
 	/** Accessor for the OnNotifySourcesOfDeviceChange delegate list */
 	FOnNotifySourcesOfDeviceChange& GetOnNotifySourcesOfDeviceChange() { return OnNotifySourcesOfDeviceChange; }
 
diff --git a/Engine/Plugins/VirtualProduction/Takes/Source/TakesCore/Public/TakeRecorderSourceProperty.h b/Engine/Plugins/VirtualProduction/Takes/Source/TakesCore/Public/TakeRecorderSourceProperty.h
--- a/Engine/Plugins/VirtualProduction/Takes/Source/TakesCore/Public/TakeRecorderSourceProperty.h
+++ b/Engine/Plugins/VirtualProduction/Takes/Source/TakesCore/Public/TakeRecorderSourceProperty.h
@@ -158,10 +158,6 @@ struct TAKESCORE_API FAudioInputDeviceProperty
 	/** The unique id of the currently selected audio device */
 	UPROPERTY(config, EditAnywhere, BlueprintReadWrite, Category = "Property", meta = (DisplayName = "Audio Input Device"))
 	FString DeviceId;
-
-	/** The number of input channels supported by this device */
-	UPROPERTY(Transient, EditAnywhere, BlueprintReadWrite, Category = "Property")
-	int32 DeviceInputChannelCount = 0;
 	
 	/** The desired buffer size used for audio callbacks during record */
 	UPROPERTY(config, EditAnywhere, BlueprintReadWrite, Category = "Property", meta = (ClampMin = "256", UIMin = "256", ClampMax = "8192", UIMax = "8192"))
