Initialize audio settings at Take Recorder instantiation rather than in details customization.

diff --git a/Engine/Plugins/VirtualProduction/Takes/Source/TakeRecorder/Private/TakeRecorderProjectSettingsCustomization.h b/Engine/Plugins/VirtualProduction/Takes/Source/TakeRecorder/Private/TakeRecorderProjectSettingsCustomization.h
--- a/Engine/Plugins/VirtualProduction/Takes/Source/TakeRecorder/Private/TakeRecorderProjectSettingsCustomization.h
+++ b/Engine/Plugins/VirtualProduction/Takes/Source/TakeRecorder/Private/TakeRecorderProjectSettingsCustomization.h
@@ -16,13 +16,6 @@ class FTakeRecorderProjectSettingsCustomization : public IDetailCustomization
 {
 	virtual void CustomizeDetails(IDetailLayoutBuilder& DetailLayout) override
 	{
-		// Enumerate audio devices before building the UI. Note, this can be expensive depending on the hardware
-		// attached to the machine, so we wait as late as possible before enumerating.
-		if (UTakeRecorderAudioInputSettings* AudioInputSettings = TakeRecorderAudioSettingsUtils::GetTakeRecorderAudioInputSettings())
-		{
-			AudioInputSettings->EnumerateAudioDevices();
-		}
-
 		FPropertyEditorModule& PropertyEditorModule = FModuleManager::Get().LoadModuleChecked<FPropertyEditorModule>("PropertyEditor");
 
 		// Register audio device details customization for the project settings window
diff --git a/Engine/Plugins/VirtualProduction/Takes/Source/TakeRecorder/Private/Widgets/SLevelSequenceTakeEditor.cpp b/Engine/Plugins/VirtualProduction/Takes/Source/TakeRecorder/Private/Widgets/SLevelSequenceTakeEditor.cpp
--- a/Engine/Plugins/VirtualProduction/Takes/Source/TakeRecorder/Private/Widgets/SLevelSequenceTakeEditor.cpp
+++ b/Engine/Plugins/VirtualProduction/Takes/Source/TakeRecorder/Private/Widgets/SLevelSequenceTakeEditor.cpp
@@ -123,6 +123,7 @@ void SLevelSequenceTakeEditor::Construct(const FArguments& InArgs)
 	.OnSelectionChanged(this, &SLevelSequenceTakeEditor::OnSourcesSelectionChanged);
 
 	CheckForNewLevelSequence();
+	InitializeAudioSettings();
 
 	ChildSlot
 	[
@@ -193,6 +194,16 @@ void SLevelSequenceTakeEditor::CheckForNewLevelSequence()
 	}
 }
 
+void SLevelSequenceTakeEditor::InitializeAudioSettings()
+{
+	// Enumerate audio devices before building the UI. Note, this can be expensive depending on the hardware
+	// attached to the machine, so we wait as late as possible before enumerating.
+	if (UTakeRecorderAudioInputSettings* AudioInputSettings = TakeRecorderAudioSettingsUtils::GetTakeRecorderAudioInputSettings())
+	{
+		AudioInputSettings->EnumerateAudioDevices();
+	}
+}
+
 void SLevelSequenceTakeEditor::Tick(const FGeometry& AllottedGeometry, const double InCurrentTime, const float InDeltaTime)
 {
 	CheckForNewLevelSequence();
@@ -791,9 +802,6 @@ private:
 					InputChannelTitleBlock->SetText(FText::AsNumber(*InSelection));
 				}
 			})
-			.OnComboBoxOpening_Lambda([this]()
-			{
-			})
 			[
 				SAssignNew(InputChannelTitleBlock, STextBlock)
 				.Text_Lambda([this]()
diff --git a/Engine/Plugins/VirtualProduction/Takes/Source/TakeRecorder/Private/Widgets/SLevelSequenceTakeEditor.h b/Engine/Plugins/VirtualProduction/Takes/Source/TakeRecorder/Private/Widgets/SLevelSequenceTakeEditor.h
--- a/Engine/Plugins/VirtualProduction/Takes/Source/TakeRecorder/Private/Widgets/SLevelSequenceTakeEditor.h
+++ b/Engine/Plugins/VirtualProduction/Takes/Source/TakeRecorder/Private/Widgets/SLevelSequenceTakeEditor.h
@@ -62,6 +62,12 @@ private:
 	 */
 	void CheckForNewLevelSequence();
 
+	/**
+	 * Initialize the audio settings object so that microphone sources can build their input channel menus
+	 * based on the current audio device
+	 */
+	void InitializeAudioSettings();
+
 	void AddDetails(const TPair<const UClass*, TArray<UObject*> >& Pair, TArray<FObjectKey>& PreviousClasses);
 
 	/**
