Small tweak to sort predicate so that Take Recorder sources with the same name land in the same order they were created.

diff --git a/Engine/Plugins/VirtualProduction/Takes/Source/TakeRecorder/Private/Widgets/STakeRecorderSources.cpp b/Engine/Plugins/VirtualProduction/Takes/Source/TakeRecorder/Private/Widgets/STakeRecorderSources.cpp
--- a/Engine/Plugins/VirtualProduction/Takes/Source/TakeRecorder/Private/Widgets/STakeRecorderSources.cpp
+++ b/Engine/Plugins/VirtualProduction/Takes/Source/TakeRecorder/Private/Widgets/STakeRecorderSources.cpp
@@ -556,7 +556,7 @@ void STakeRecorderSources::ReconstructTree()
 
 		auto SortPredicate = [](TTuple<int32, FString>& A, TTuple<int32, FString>& B)
 		{
-			return A.Get<1>() < B.Get<1>();
+			return A.Get<1>() <= B.Get<1>();
 		};
 
 		Algo::Sort(SortData, SortPredicate);
