Take Recorder: Fix range for camera cut section so that it works for subsequences

diff --git a/Engine/Plugins/VirtualProduction/Takes/Source/TakeRecorderSources/Private/TakeRecorderActorSource.cpp b/Engine/Plugins/VirtualProduction/Takes/Source/TakeRecorderSources/Private/TakeRecorderActorSource.cpp
--- a/Engine/Plugins/VirtualProduction/Takes/Source/TakeRecorderSources/Private/TakeRecorderActorSource.cpp
+++ b/Engine/Plugins/VirtualProduction/Takes/Source/TakeRecorderSources/Private/TakeRecorderActorSource.cpp
@@ -822,7 +822,7 @@ TArray<UTakeRecorderSource*> UTakeRecorderActorSource::PostRecording(ULevelSeque
 			{
 				FGuid RecordedCameraGuid = GetRecordedActorGuid(Target.Get());
 				FMovieSceneSequenceID RecordedCameraSequenceID = GetLevelSequenceID(Target.Get());
-				TakesUtils::CreateCameraCutTrack(InRootSequence, RecordedCameraGuid, RecordedCameraSequenceID, InSequence->GetMovieScene()->GetPlaybackRange());
+				TakesUtils::CreateCameraCutTrack(InRootSequence, RecordedCameraGuid, RecordedCameraSequenceID, InRootSequence->GetMovieScene()->GetPlaybackRange());
 			}
 
 			// Swap our target actor to the Editor actor (in case the recording was added while in PIE)
