Cleanup of inconsistent asset naming and various other small issues related to Take Recorder audio.

diff --git a/Engine/Plugins/Runtime/AudioCapture/Source/AudioCaptureEditor/Private/AudioRecordingManager.cpp b/Engine/Plugins/Runtime/AudioCapture/Source/AudioCaptureEditor/Private/AudioRecordingManager.cpp
--- a/Engine/Plugins/Runtime/AudioCapture/Source/AudioCaptureEditor/Private/AudioRecordingManager.cpp
+++ b/Engine/Plugins/Runtime/AudioCapture/Source/AudioCaptureEditor/Private/AudioRecordingManager.cpp
@@ -350,7 +350,6 @@ TObjectPtr<USoundWave> FAudioRecordingManager::CreateSoundWaveAsset(const FRecor
 	{
 		FString PackageName;
 		FString AssetName = InSourceSettings.AssetName;
-		AssetName = FString::Printf(TEXT("%s_Channel_%d"), *AssetName, InSourceSettings.InputChannelNumber);
 		PackageName = InSourceSettings.Directory.Path / AssetName;
 
 		int32 ChannelIndex = InSourceSettings.InputChannelNumber - 1;
@@ -434,7 +433,7 @@ FSoundWaveTimecodeInfo FAudioRecordingManager::GetTimecodeInfo(const FRecordingM
 	FSoundWaveTimecodeInfo TimecodeInfo;
 
 	double SecondsFromMidnight = InSourceSettings.StartTimecode.ToTimespan(InSourceSettings.VideoFrameRate).GetTotalSeconds();
-	if (!ensureAlwaysMsgf(SecondsFromMidnight > 0, TEXT("Invalid SecondsFromMidnight value: %f, defaulting to 0"), SecondsFromMidnight))
+	if (!ensureAlwaysMsgf(SecondsFromMidnight >= 0, TEXT("Invalid SecondsFromMidnight value: %f, defaulting to 0"), SecondsFromMidnight))
 	{
 		SecondsFromMidnight = 0.0;
 	}
diff --git a/Engine/Plugins/VirtualProduction/Takes/Source/TakeRecorderSources/Private/TakeRecorderMicrophoneAudioSource.cpp b/Engine/Plugins/VirtualProduction/Takes/Source/TakeRecorderSources/Private/TakeRecorderMicrophoneAudioSource.cpp
--- a/Engine/Plugins/VirtualProduction/Takes/Source/TakeRecorderSources/Private/TakeRecorderMicrophoneAudioSource.cpp
+++ b/Engine/Plugins/VirtualProduction/Takes/Source/TakeRecorderSources/Private/TakeRecorderMicrophoneAudioSource.cpp
@@ -29,11 +29,13 @@
 
 UTakeRecorderMicrophoneAudioSourceSettings::UTakeRecorderMicrophoneAudioSourceSettings(const FObjectInitializer& ObjInit)
 	: Super(ObjInit)
+	, AudioSourceName(NSLOCTEXT("UTakeRecorderMicrophoneAudioSource", "Label", "Microphone Audio"))
 	, AudioTrackName(NSLOCTEXT("UTakeRecorderMicrophoneAudioSource", "DefaultAudioTrackName", "Recorded Audio"))
-	, AudioAssetName(TEXT("Audio_{slate}_{take}"))
+	, AudioAssetName(TEXT("Audio_{slate}_{take}_Channel_{channel}"))
 	, AudioSubDirectory(TEXT("Audio"))
 {
 	TrackTint = FColor(75, 67, 148);
+	AudioSourceName = FText::Format(FText::FromString(TEXT("{0} {1}")), AudioSourceName, FText::FromString(TEXT("{channel}")));
 }
 
 void UTakeRecorderMicrophoneAudioSourceSettings::PostEditChangeProperty(FPropertyChangedEvent& PropertyChangedEvent)
@@ -145,6 +147,57 @@ void UTakeRecorderMicrophoneAudioSource::SetAudioDeviceChannelCount(int32 InChan
 	}
 }
 
+FString UTakeRecorderMicrophoneAudioSource::ReplaceStringTokens(const FString& InString) const
+{
+	TMap<FString, FStringFormatArg> FormatArgs;
+	FormatArgs.Add(TEXT("channel"), FString::FromInt(AudioChannel.AudioInputDeviceChannel));
+
+	return FString::Format(*InString, FormatArgs);
+}
+
+FString UTakeRecorderMicrophoneAudioSource::GetAudioTrackName(ULevelSequence* InSequence) const
+{
+	if (UTakeMetaData* TakeMetaData = InSequence->FindMetaData<UTakeMetaData>())
+	{
+		FString TempStr = *TakeMetaData->GenerateAssetPath(*AudioTrackName.ToString());
+
+		return ReplaceStringTokens(TempStr);
+	}
+
+	return TEXT("MicrophoneAudio");
+}
+
+FString UTakeRecorderMicrophoneAudioSource::GetAudioAssetName(ULevelSequence* InSequence) const
+{
+	// This is called for both uassets and subsequences, both of which require unique names.
+	// Append the channel number if it's not already part of the name to ensure uniqueness.
+
+	if (UTakeMetaData* TakeMetaData = InSequence->FindMetaData<UTakeMetaData>())
+	{
+		bool bHasChannelCount = AudioAssetName.Contains("{channel}");
+		FString TempStr = ReplaceStringTokens(TakeMetaData->GenerateAssetPath(AudioAssetName));
+
+		if (!bHasChannelCount)
+		{
+			TempStr = FString::Printf(TEXT("%s_Channel_%d"), *TempStr, AudioChannel.AudioInputDeviceChannel);
+		}
+
+		return TempStr.Replace(TEXT(" "), TEXT("_"));
+	}
+
+	return TEXT("MicrophoneAudio");
+}
+
+FString UTakeRecorderMicrophoneAudioSource::GetSubsceneTrackName(ULevelSequence* InSequence) const
+{
+	return GetAudioTrackName(InSequence);
+}
+
+FString UTakeRecorderMicrophoneAudioSource::GetSubsceneAssetName(ULevelSequence* InSequence) const
+{
+	return GetAudioAssetName(InSequence);
+}
+
 static FString MakeNewAssetName(const FString& BaseAssetPath, const FString& BaseAssetName)
 {
 	const FString Dot(TEXT("."));
@@ -175,23 +228,48 @@ TArray<UTakeRecorderSource*> UTakeRecorderMicrophoneAudioSource::PreRecording(UL
 {
 	if (AudioChannel.AudioInputDeviceChannel > 0)
 	{
+		FString TrackName = GetAudioTrackName(InSequence);
 		UMovieScene* MovieScene = InSequence->GetMovieScene();
-		for (auto Track : MovieScene->GetTracks())
+		
+		for (auto Track : ReverseIterate(MovieScene->GetTracks()))
 		{
-			if (Track->IsA(UMovieSceneAudioTrack::StaticClass()) && Track->GetDisplayName().EqualTo(AudioTrackName))
+			if (Track->IsA(UMovieSceneAudioTrack::StaticClass()) && Track->GetDisplayName().EqualTo(FText::FromString(TrackName)))
 			{
-				CachedAudioTrack = Cast<UMovieSceneAudioTrack>(Track);
+				UMovieSceneAudioTrack* AudioTrack = Cast<UMovieSceneAudioTrack>(Track);
+				if (!IsTrackAssociatedWithAnySource(AudioTrack))
+				{
+					CachedAudioTrack = AudioTrack;
+					break;
+				}
 			}
 		}
 
 		if (!CachedAudioTrack.IsValid())
 		{
 			CachedAudioTrack = MovieScene->AddTrack<UMovieSceneAudioTrack>();
-			CachedAudioTrack->SetDisplayName(AudioTrackName);
+			CachedAudioTrack->SetDisplayName(FText::FromString(TrackName));
+		}
+
+		if (bReplaceRecordedAudio)
+		{
+			CachedAudioTrack->RemoveAllAnimationData();
 		}
 
 		FString PathToRecordTo = FPackageName::GetLongPackagePath(InSequence->GetOutermost()->GetPathName());
-		FString BaseName = InSequence->GetName();
+		FString BaseName;
+		UTakeRecorderSources* Sources = GetTypedOuter<UTakeRecorderSources>();
+
+		// When recording into subsequences, the record path is appended to the slate string
+		// which results in inconsistent asset names. Here we work around that by using the
+		// subsequence name which will be the same as the asset name in the non-subsequence case.
+		if (Sources && Sources->GetSettings().bRecordSourcesIntoSubSequences)
+		{
+			BaseName = InSequence->GetName();
+		}
+		else
+		{
+			BaseName = GetAudioAssetName(InSequence);
+		}
 
 		AudioDirectory.Path = PathToRecordTo;
 		if (AudioSubDirectory.Len())
@@ -200,9 +278,13 @@ TArray<UTakeRecorderSource*> UTakeRecorderMicrophoneAudioSource::PreRecording(UL
 		}
 
 		AssetName = MakeNewAssetName(AudioDirectory.Path, BaseName);
-	}
 
-	RecordedSoundWaves.Empty();
+		// Add the section here so it is displayed during record (non-subsequence case)
+		UMovieSceneAudioSection* NewAudioSection = NewObject<UMovieSceneAudioSection>(CachedAudioTrack.Get(), UMovieSceneAudioSection::StaticClass());
+		NewAudioSection->SetRowIndex(0);
+		NewAudioSection->SetLooping(false);
+		CachedAudioTrack->AddSection(*NewAudioSection);
+	}
 
 	return TArray<UTakeRecorderSource*>();
 }
@@ -222,6 +304,12 @@ void UTakeRecorderMicrophoneAudioSource::StartRecording(const FTimecode& InSecti
 
 	StartTimecode = InSectionStartTimecode;
 
+	const TArray<UMovieSceneSection*>& AudioSections = CachedAudioTrack->GetAudioSections();
+	if (AudioSections.Num() > 0)
+	{
+		AudioSections[0]->SetRange(TRange<FFrameNumber>(InSectionFirstFrame, InSectionFirstFrame));
+	}
+
 	UTakeRecorderMicrophoneAudioManager* AudioInputManager = GetAudioInputManager();
 	if (AudioInputManager != nullptr)
 	{
@@ -240,6 +328,20 @@ void UTakeRecorderMicrophoneAudioSource::StartRecording(const FTimecode& InSecti
 	}
 }
 
+void UTakeRecorderMicrophoneAudioSource::TickRecording(const FQualifiedFrameTime& CurrentTime)
+{
+	if (CachedAudioTrack.IsValid())
+	{
+		FFrameRate   TickResolution = CachedAudioTrack->GetTypedOuter<UMovieScene>()->GetTickResolution();
+		FFrameNumber CurrentFrame = CurrentTime.ConvertTo(TickResolution).FloorToFrame();
+
+		for (UMovieSceneSection* Section : CachedAudioTrack->GetAllSections())
+		{
+			Section->ExpandToFrame(CurrentFrame);
+		}
+	}
+}
+
 void UTakeRecorderMicrophoneAudioSource::StopRecording(class ULevelSequence* InSequence)
 {
 	Super::StopRecording(InSequence);
@@ -255,7 +357,7 @@ TArray<UTakeRecorderSource*> UTakeRecorderMicrophoneAudioSource::PostRecording(U
 {
 	GetRecordedSoundWave(InSequence);
 
-	if (!RecordedSoundWaves.Num())
+	if (!RecordedSoundWave.IsValid())
 	{
 		return TArray<UTakeRecorderSource*>();
 	}
@@ -263,71 +365,48 @@ TArray<UTakeRecorderSource*> UTakeRecorderMicrophoneAudioSource::PostRecording(U
 	TArray<UObject*> AssetsToCleanUp;
 	if (bCancelled)
 	{
-		for (TWeakObjectPtr<USoundWave> WeakRecordedSoundWave : RecordedSoundWaves)
+		if (USoundWave* SoundWave = RecordedSoundWave.Get()) 
 		{
-			if (USoundWave* RecordedSoundWave = WeakRecordedSoundWave.Get())
-			{
-				AssetsToCleanUp.Add(RecordedSoundWave);
-			}
+			AssetsToCleanUp.Add(SoundWave);
 		}
 	}
 	else
 	{
-		for (TWeakObjectPtr<USoundWave> WeakRecordedSoundWave : RecordedSoundWaves)
+		if (USoundWave* SoundWave = RecordedSoundWave.Get())
 		{
-			if (USoundWave* RecordedSoundWave = WeakRecordedSoundWave.Get())
-			{
-				RecordedSoundWave->MarkPackageDirty();
-		
-				FAssetRegistryModule::AssetCreated(RecordedSoundWave);
-			}
-		}
-
-		UMovieScene* MovieScene = InSequence->GetMovieScene();
-		check(CachedAudioTrack.IsValid());
-
-		FFrameRate TickResolution = MovieScene->GetTickResolution();
-		FFrameRate DisplayRate = MovieScene->GetDisplayRate();
+			SoundWave->MarkPackageDirty();
 
-		if (bReplaceRecordedAudio)
-		{
-			CachedAudioTrack->RemoveAllAnimationData();
+			FAssetRegistryModule::AssetCreated(SoundWave);
 		}
 
+		UMovieScene* MovieScene = InSequence->GetMovieScene();
 		UTakeRecorderSources* Sources = GetTypedOuter<UTakeRecorderSources>();
 
-		for (TWeakObjectPtr<USoundWave> WeakRecordedSoundWave : RecordedSoundWaves)
+		if (ensure(CachedAudioTrack.IsValid()) && ensure(Sources != nullptr))
 		{
-			if (USoundWave* RecordedSoundWave = WeakRecordedSoundWave.Get())
+			if (USoundWave* SoundWave = RecordedSoundWave.Get())
 			{
-				int32 RowIndex = -1;
-				for (UMovieSceneSection* Section : CachedAudioTrack->GetAllSections())
+				const TArray<UMovieSceneSection*>& AudioSections = CachedAudioTrack->GetAudioSections();
+				if (AudioSections.Num() > 0)
 				{
-					RowIndex = FMath::Max(RowIndex, Section->GetRowIndex());
+					UMovieSceneAudioSection* AudioSection = Cast<UMovieSceneAudioSection>(AudioSections[0]);					
+					AudioSection->SetSound(SoundWave);
+					AudioSection->TimecodeSource = StartTimecode;
 				}
 
-				UMovieSceneAudioSection* NewAudioSection = NewObject<UMovieSceneAudioSection>(CachedAudioTrack.Get(), UMovieSceneAudioSection::StaticClass());
-
-				FFrameNumber RecordStartFrame = MovieScene->GetPlaybackRange().GetLowerBoundValue();
-
-				NewAudioSection->SetRowIndex(RowIndex + 1);
-				NewAudioSection->SetSound(RecordedSoundWave);
-				NewAudioSection->SetRange(TRange<FFrameNumber>(RecordStartFrame, RecordStartFrame + (RecordedSoundWave->GetDuration() * TickResolution).CeilToFrame()));
-				NewAudioSection->TimecodeSource = FTimecode::FromFrameNumber(RecordStartFrame, DisplayRate);
-
-				CachedAudioTrack->AddSection(*NewAudioSection);
-
-				if ((Sources && Sources->GetSettings().bSaveRecordedAssets) || GEditor == nullptr)
+				if ((Sources->GetSettings().bSaveRecordedAssets) || GEditor == nullptr)
 				{
-					TakesUtils::SaveAsset(RecordedSoundWave);
+					TakesUtils::SaveAsset(SoundWave);
 				}
 			}
 		}
 	}
 
 	// Reset our audio track pointer
-	CachedAudioTrack = nullptr;
-	RecordedSoundWaves.Empty();
+	CachedAudioTrack.Reset();
+	
+	// Reset our sound wave pointer
+	RecordedSoundWave.Reset();
 	
 	if (GEditor && AssetsToCleanUp.Num() > 0)
 	{
@@ -371,14 +450,15 @@ void UTakeRecorderMicrophoneAudioSource::GetRecordedSoundWave(ULevelSequence* In
 			TimecodeOffset.NumOfSecondsSinceMidnight = StartTimecode.ToTimespan(DisplayRate).GetTotalSeconds();
 			SoundWave->SetTimecodeOffset(TimecodeOffset);
 			
-			RecordedSoundWaves.Add(SoundWave);
+			RecordedSoundWave = SoundWave;
 		}
 	}
 }
 
 FText UTakeRecorderMicrophoneAudioSource::GetDisplayTextImpl() const
 {
-	return NSLOCTEXT("UTakeRecorderMicrophoneAudioSource", "Label", "Microphone Audio");
+
+	return FText::FromString(ReplaceStringTokens(AudioSourceName.ToString()));
 }
 
 bool UTakeRecorderMicrophoneAudioSource::CanAddSource(UTakeRecorderSources* InSources) const
@@ -401,3 +481,20 @@ void UTakeRecorderMicrophoneAudioSource::SetCurrentInputChannel(int32 InChannelN
 		AudioChannel.AudioInputDeviceChannel = InChannelNumber;
 	}
 }
+
+bool UTakeRecorderMicrophoneAudioSource::IsTrackAssociatedWithAnySource(UMovieSceneAudioTrack* InAudioTrack)
+{
+	UTakeRecorderSources* SourcesList = GetTypedOuter<UTakeRecorderSources>();
+	for (UTakeRecorderSource* Source : SourcesList->GetSources())
+	{
+		if (UTakeRecorderMicrophoneAudioSource* MicSource = Cast<UTakeRecorderMicrophoneAudioSource>(Source))
+		{
+			if (MicSource->CachedAudioTrack == InAudioTrack)
+			{
+				return true;
+			}
+		}
+	}
+
+	return false;
+}
diff --git a/Engine/Plugins/VirtualProduction/Takes/Source/TakeRecorderSources/Private/TakeRecorderMicrophoneAudioSource.h b/Engine/Plugins/VirtualProduction/Takes/Source/TakeRecorderSources/Private/TakeRecorderMicrophoneAudioSource.h
--- a/Engine/Plugins/VirtualProduction/Takes/Source/TakeRecorderSources/Private/TakeRecorderMicrophoneAudioSource.h
+++ b/Engine/Plugins/VirtualProduction/Takes/Source/TakeRecorderSources/Private/TakeRecorderMicrophoneAudioSource.h
@@ -12,6 +12,7 @@
 
 #include "TakeRecorderMicrophoneAudioSource.generated.h"
 
+class UMovieSceneAudioTrack;
 class USoundWave;
 class UTakeRecorderMicrophoneAudioManager;
 struct FTakeRecorderAudioSourceSettings;
@@ -32,6 +33,10 @@ public:
 	virtual FString GetSubsceneAssetName(ULevelSequence* InSequence) const override;
 	// ~UTakeRecorderSource Interface
 
+	/** Name of the audio source */
+	UPROPERTY(config, EditAnywhere, BlueprintReadWrite, Category = "Source")
+	FText AudioSourceName;
+
 	/** Name of the recorded audio track */
 	UPROPERTY(config, EditAnywhere, BlueprintReadWrite, Category = "Source")
 	FText AudioTrackName;
@@ -59,7 +64,7 @@ public:
 	// End UObject Interface
 
 	/** Gain in decibels to apply to recorded audio */
-	UPROPERTY(EditAnywhere, BlueprintReadWrite, Category = "Source", meta = (ClampMin = "0.0", UIMin = "0.0"))
+	UPROPERTY(EditAnywhere, BlueprintReadWrite, Category = "Source", meta = (ClampMin = "0.0", UIMin = "0.0", UIMax = "20.0", ClampMax = "40.0"))
 	float AudioGain;
 
 	/** Whether or not to split mic channels into separate audio tracks. If not true, a max of 2 input channels is supported. */
@@ -78,8 +83,11 @@ public:
 
 	// UTakeRecorderSource
 	virtual void Initialize() override;
+	virtual FString GetSubsceneTrackName(ULevelSequence* InSequence) const override;
+	virtual FString GetSubsceneAssetName(ULevelSequence* InSequence) const override;
 	virtual TArray<UTakeRecorderSource*> PreRecording(ULevelSequence* InSequence, FMovieSceneSequenceID InSequenceID, ULevelSequence* InRootSequence, FManifestSerializer* InManifestSerializer) override;
 	virtual void StartRecording(const FTimecode& InSectionStartTimecode, const FFrameNumber& InSectionFirstFrame, class ULevelSequence* InSequence) override;
+	virtual void TickRecording(const FQualifiedFrameTime& CurrentTime) override;
 	virtual void StopRecording(class ULevelSequence* InSequence) override;
 	virtual TArray<UTakeRecorderSource*> PostRecording(ULevelSequence* InSequence, class ULevelSequence* InRootSequence, const bool bCancelled) override;
 	virtual void FinalizeRecording() override;
@@ -88,6 +96,11 @@ public:
 	virtual bool CanAddSource(UTakeRecorderSources* InSources) const override;
 	// ~UTakeRecorderSource
 
+	/** Returns track display name, replacing any tokens if needed */
+	FString GetAudioTrackName(ULevelSequence* InSequence) const;
+	/** Returns the fully expanded asset name */
+	FString GetAudioAssetName(ULevelSequence* InSequence) const;
+
 	/** Sets the channel count supported by the currently selected audio device. */
 	void SetAudioDeviceChannelCount(int32 InChannelCount);
 	/** Delegate which receives notifications when the audio input device changes. */
@@ -98,6 +111,9 @@ private:
 	/** Helper function for getting a pointer to the AudioInputManger object. */
 	static UTakeRecorderMicrophoneAudioManager* GetAudioInputManager();
 	
+	/** Calls ReplaceSourceTokens and ReplaceBaseTokens in an effort to expand all embedded tokens in the given string */
+	FString ReplaceStringTokens(const FString& InString) const;
+
 	/** Returns an array of booleans indicating which channel indexes are currently in use */
 	TUniquePtr<TArray<bool>> GetChannelsInUse(const int32 InDeviceChannelCount);
 	/** Fetches the USoundWave for this source after a Take has been recorded */
@@ -105,12 +121,15 @@ private:
 	/** Called when the user selects a new channel in the combobox for this source */
 	void SetCurrentInputChannel(int32 InChannelNumber);
 
+	/** Returns whether given track is being used by any other source */
+	bool IsTrackAssociatedWithAnySource(UMovieSceneAudioTrack* InAudioTrack);
+
 private:
 
 	// Holds the Sequencer audio track for this source
 	TWeakObjectPtr<class UMovieSceneAudioTrack> CachedAudioTrack;
-	// Holds the USoundWave assets for a given take
-	TArray<TWeakObjectPtr<USoundWave>> RecordedSoundWaves;
+	// Holds the USoundWave asset for a given take
+	TWeakObjectPtr<USoundWave> RecordedSoundWave;
 	// Caches the starting timecode for this take so it can be referenced when creating USoundWave assets
 	FTimecode StartTimecode;
 
diff --git a/Engine/Source/Editor/MovieSceneTools/Private/TrackEditors/AudioTrackEditor.cpp b/Engine/Source/Editor/MovieSceneTools/Private/TrackEditors/AudioTrackEditor.cpp
--- a/Engine/Source/Editor/MovieSceneTools/Private/TrackEditors/AudioTrackEditor.cpp
+++ b/Engine/Source/Editor/MovieSceneTools/Private/TrackEditors/AudioTrackEditor.cpp
@@ -651,9 +651,18 @@ UMovieSceneSection* FAudioSection::GetSectionObject()
 FText FAudioSection::GetSectionTitle() const
 {
 	UMovieSceneAudioSection* AudioSection = Cast<UMovieSceneAudioSection>(&Section);
-	if (AudioSection && AudioSection->GetSound())
+	if (AudioSection)
 	{
-		return FText::FromString(AudioSection->GetSound()->GetName());
+		if (AudioSection->GetSound())
+		{
+			// Return the asset name if it exists
+			return FText::FromString(AudioSection->GetSound()->GetName());
+		}
+		else
+		{
+			// There is no asset during record so return empty string
+			return FText();
+		}
 	}
 	
 	return NSLOCTEXT("FAudioSection", "NoAudioTitleName", "No Audio");
