When recording timecode, set timecode source value to incoming timecode at record start.

diff --git a/Engine/Plugins/VirtualProduction/Takes/Source/TakeRecorderSources/Private/TakeRecorderMicrophoneAudioSource.cpp b/Engine/Plugins/VirtualProduction/Takes/Source/TakeRecorderSources/Private/TakeRecorderMicrophoneAudioSource.cpp
--- a/Engine/Plugins/VirtualProduction/Takes/Source/TakeRecorderSources/Private/TakeRecorderMicrophoneAudioSource.cpp
+++ b/Engine/Plugins/VirtualProduction/Takes/Source/TakeRecorderSources/Private/TakeRecorderMicrophoneAudioSource.cpp
@@ -355,6 +355,12 @@ void UTakeRecorderMicrophoneAudioSource::StopRecording(class ULevelSequence* InS
 
 TArray<UTakeRecorderSource*> UTakeRecorderMicrophoneAudioSource::PostRecording(ULevelSequence* InSequence, class ULevelSequence* InRootSequence, const bool bCancelled)
 {
+	FTakeRecorderProjectParameters ProjectParameters = GetDefault<UTakeRecorderProjectSettings>()->Settings;
+	if (ProjectParameters.bRecordTimecode)
+	{
+		ProcessRecordedTimes(InSequence);
+	}
+
 	GetRecordedSoundWave(InSequence);
 
 	if (!RecordedSoundWave.IsValid())
@@ -425,6 +431,27 @@ void UTakeRecorderMicrophoneAudioSource::FinalizeRecording()
 	}
 }
 
+void UTakeRecorderMicrophoneAudioSource::ProcessRecordedTimes(ULevelSequence* InSequence)
+{
+	if (ensure(CachedAudioTrack.IsValid()))
+	{
+		const TArray<TPair<FQualifiedFrameTime, FQualifiedFrameTime>>& RecordedTimes = UTakeRecorderSources::RecordedTimes;
+		const TArray<UMovieSceneSection*>& AudioSections = CachedAudioTrack->GetAudioSections();
+
+		TRange<FFrameNumber> FrameRange = AudioSections[0]->GetRange();
+
+		for (const TPair<FQualifiedFrameTime, FQualifiedFrameTime>& RecordedTimePair : RecordedTimes)
+		{
+			FFrameNumber FrameNumber = RecordedTimePair.Key.Time.FrameNumber;
+			if (FrameRange.Contains(FrameNumber))
+			{
+				StartTimecode = RecordedTimePair.Value.ToTimecode();
+				break;
+			}
+		}
+	}
+}
+
 void UTakeRecorderMicrophoneAudioSource::GetRecordedSoundWave(ULevelSequence* InSequence)
 {
 	UTakeRecorderMicrophoneAudioManager* AudioInputManager = GetAudioInputManager();
diff --git a/Engine/Plugins/VirtualProduction/Takes/Source/TakeRecorderSources/Private/TakeRecorderMicrophoneAudioSource.h b/Engine/Plugins/VirtualProduction/Takes/Source/TakeRecorderSources/Private/TakeRecorderMicrophoneAudioSource.h
--- a/Engine/Plugins/VirtualProduction/Takes/Source/TakeRecorderSources/Private/TakeRecorderMicrophoneAudioSource.h
+++ b/Engine/Plugins/VirtualProduction/Takes/Source/TakeRecorderSources/Private/TakeRecorderMicrophoneAudioSource.h
@@ -111,6 +111,9 @@ private:
 	/** Helper function for getting a pointer to the AudioInputManger object. */
 	static UTakeRecorderMicrophoneAudioManager* GetAudioInputManager();
 	
+	/** Parses recorded timecode values */
+	void ProcessRecordedTimes(ULevelSequence* InSequence);
+
 	/** Calls ReplaceSourceTokens and ReplaceBaseTokens in an effort to expand all embedded tokens in the given string */
 	FString ReplaceStringTokens(const FString& InString) const;
 
