Setting default audio device updates device list combo box if neccessary.

diff --git a/Engine/Plugins/VirtualProduction/Takes/Source/TakeRecorder/Private/Widgets/TakeRecorderAudioSettingsCustomization.cpp b/Engine/Plugins/VirtualProduction/Takes/Source/TakeRecorder/Private/Widgets/TakeRecorderAudioSettingsCustomization.cpp
--- a/Engine/Plugins/VirtualProduction/Takes/Source/TakeRecorder/Private/Widgets/TakeRecorderAudioSettingsCustomization.cpp
+++ b/Engine/Plugins/VirtualProduction/Takes/Source/TakeRecorder/Private/Widgets/TakeRecorderAudioSettingsCustomization.cpp
@@ -37,7 +37,7 @@ void FAudioInputDevicePropertyCustomization::CustomizeChildren(TSharedRef<IPrope
 			{
 				bool bUseSystemDefault = (NewState == ECheckBoxState::Checked);
 				UseSystemDefaultHandle->SetValue(bUseSystemDefault);
-				UpdateDeviceCombobxEnabled();
+                SynchronizeWidgetStates();
 			})
 		];
 	}
@@ -160,6 +160,11 @@ TSharedRef<SWidget> FAudioInputDevicePropertyCustomization::MakeAudioInputSelect
 			bool bIsDefaultDevice = GetIsDefaultFromInfoProperty(DeviceInfoHandle);
 			FString DeviceId = GetDeviceIdFromInfoProperty(DeviceInfoHandle);
 
+            if (bIsDefaultDevice)
+            {
+                DefaultDeviceId = DeviceId;
+            }
+            
 			if (bIsDefaultDevice && !bHaveCurrentDevice)
 			{
 				CurrentDeviceId = DeviceId;
@@ -227,6 +232,12 @@ void FAudioInputDevicePropertyCustomization::UpdateDeviceCombobxEnabled()
 	if (bEnableWidget != bIsWidgetEnabled)
 	{
 		ComboBox->SetEnabled(bEnableWidget);
+        
+        // If disabling, set the current menu item to default
+        if (!bEnableWidget)
+        {
+            InputDeviceHandle->SetValue(DefaultDeviceId);
+        }
 	}
 }
 
diff --git a/Engine/Plugins/VirtualProduction/Takes/Source/TakeRecorder/Private/Widgets/TakeRecorderAudioSettingsCustomization.h b/Engine/Plugins/VirtualProduction/Takes/Source/TakeRecorder/Private/Widgets/TakeRecorderAudioSettingsCustomization.h
--- a/Engine/Plugins/VirtualProduction/Takes/Source/TakeRecorder/Private/Widgets/TakeRecorderAudioSettingsCustomization.h
+++ b/Engine/Plugins/VirtualProduction/Takes/Source/TakeRecorder/Private/Widgets/TakeRecorderAudioSettingsCustomization.h
@@ -38,6 +38,7 @@ private:
 	TArray<TSharedPtr<IPropertyHandle>> AudioInputDevices;
 	TSharedPtr<SComboBox<TSharedPtr<IPropertyHandle>>> ComboBox;
 	TSharedPtr<STextBlock> ComboBoxTitleBlock;
+	FString DefaultDeviceId;
 };
 
 class TakeRecorderAudioSettingsUtils
