Take Recorder: Lowered ensure to log warning about missing property map for parent

diff --git a/Engine/Plugins/VirtualProduction/Takes/Source/TakeRecorderSources/Private/TakeRecorderActorSource.cpp b/Engine/Plugins/VirtualProduction/Takes/Source/TakeRecorderSources/Private/TakeRecorderActorSource.cpp
--- a/Engine/Plugins/VirtualProduction/Takes/Source/TakeRecorderSources/Private/TakeRecorderActorSource.cpp
+++ b/Engine/Plugins/VirtualProduction/Takes/Source/TakeRecorderSources/Private/TakeRecorderActorSource.cpp
@@ -578,10 +578,14 @@ void UTakeRecorderActorSource::TickRecording(const FQualifiedFrameTime& CurrentS
 
 			// Add the new property map as a child of the correct parent, otherwise recursion doesn't work when we try to update the cached number of recorded properties.
 			UActorRecorderPropertyMap* ParentPropertyMap = GetParentPropertyMapForComponent(AddedComponent);
-			if(ensureMsgf(ParentPropertyMap, TEXT("A component %s was added to actor %s at runtime but we couldn't find the property map for the parent. Is the parent no longer valid?"), *AddedComponent->GetName(), *Target->GetName()))
+			if (ParentPropertyMap)
 			{
 				ParentPropertyMap->Children.Add(ComponentPropertyMap);
 			}
+			else
+			{
+				UE_LOG(LogTakesCore, Log, TEXT("A component %s was added to actor %s at runtime but we couldn't find the property map for the parent. Is the parent no longer valid?"), *AddedComponent->GetName(), *Target->GetName());
+			}
 
 			// Create the Property Map
 			RebuildRecordedPropertyMapRecursive(AddedComponent, ComponentPropertyMap);
