Fix for missing default audio device ID for new Take Recorder project.

diff --git a/Engine/Plugins/VirtualProduction/Takes/Source/TakeRecorderSources/Private/TakeRecorderMicrophoneAudioManager.cpp b/Engine/Plugins/VirtualProduction/Takes/Source/TakeRecorderSources/Private/TakeRecorderMicrophoneAudioManager.cpp
--- a/Engine/Plugins/VirtualProduction/Takes/Source/TakeRecorderSources/Private/TakeRecorderMicrophoneAudioManager.cpp
+++ b/Engine/Plugins/VirtualProduction/Takes/Source/TakeRecorderSources/Private/TakeRecorderMicrophoneAudioManager.cpp
@@ -132,13 +132,13 @@ void UTakeRecorderMicrophoneAudioManager::BuildDeviceInfoArray()
 	AudioInputDevice.DeviceInfoArray.Empty();
 
 	FTakeRecorderAudioDeviceInfo DefaultDeviceInfo;
-	bool bFoundDefaultDevice = AudioRecorder->GetCaptureDeviceInfo(DefaultDeviceInfo);
+	const bool bFoundDefaultDevice = AudioRecorder->GetCaptureDeviceInfo(DefaultDeviceInfo);
 
 	TArray<FTakeRecorderAudioDeviceInfo> CaptureDevicesAvailable;
 	AudioRecorder->GetCaptureDevicesAvailable(CaptureDevicesAvailable);
 	for (const FTakeRecorderAudioDeviceInfo& DeviceInfo : CaptureDevicesAvailable)
 	{
-		bool bIsDefaultDevice = bFoundDefaultDevice && (DeviceInfo.DeviceId == DefaultDeviceInfo.DeviceId);
+		const bool bIsDefaultDevice = bFoundDefaultDevice && (DeviceInfo.DeviceId == DefaultDeviceInfo.DeviceId);
 		AudioInputDevice.DeviceInfoArray.Add(FAudioInputDeviceInfoProperty(DeviceInfo.DeviceName, DeviceInfo.DeviceId, DeviceInfo.InputChannels, DeviceInfo.PreferredSampleRate, bIsDefaultDevice));
 	}
 
@@ -192,4 +192,10 @@ void UTakeRecorderMicrophoneAudioManager::BuildDeviceInfoArray()
 
 			return LeftLen < RightLen;
 		});
+
+	// Default to using the default input device if not already set
+	if (bFoundDefaultDevice && AudioInputDevice.DeviceId.IsEmpty())
+	{
+		AudioInputDevice.DeviceId = DefaultDeviceInfo.DeviceId;
+	}
 }
