Better error handling when previously saved audio input device is not available.

diff --git a/Engine/Plugins/Runtime/AudioCapture/Source/AudioCaptureEditor/Private/AudioRecordingManager.cpp b/Engine/Plugins/Runtime/AudioCapture/Source/AudioCaptureEditor/Private/AudioRecordingManager.cpp
--- a/Engine/Plugins/Runtime/AudioCapture/Source/AudioCaptureEditor/Private/AudioRecordingManager.cpp
+++ b/Engine/Plugins/Runtime/AudioCapture/Source/AudioCaptureEditor/Private/AudioRecordingManager.cpp
@@ -228,12 +228,15 @@ TObjectPtr<USoundWave> FAudioRecordingManager::GetRecordedSoundWave(const FRecor
 			ProcessRecordedData(InSourceSettings);
 		}
 
-		if (RecordedSoundWaves[InSourceSettings.InputChannelNumber - 1] != nullptr)
+		const int32 ChannelIndex = InSourceSettings.InputChannelNumber - 1;
+		if (RecordedSoundWaves.IsValidIndex(ChannelIndex) && RecordedSoundWaves[ChannelIndex] != nullptr)
 		{
-			return RecordedSoundWaves[InSourceSettings.InputChannelNumber - 1];
+			return RecordedSoundWaves[ChannelIndex];
+		}
+		else if (NumRecordedSamples > 0)
+		{
+			return CreateSoundWaveAsset(InSourceSettings);
 		}
-
-		return CreateSoundWaveAsset(InSourceSettings);
 	}
 
 	return nullptr;
diff --git a/Engine/Plugins/VirtualProduction/Takes/Source/TakeRecorderSources/Private/TakeRecorderMicrophoneAudioManager.cpp b/Engine/Plugins/VirtualProduction/Takes/Source/TakeRecorderSources/Private/TakeRecorderMicrophoneAudioManager.cpp
--- a/Engine/Plugins/VirtualProduction/Takes/Source/TakeRecorderSources/Private/TakeRecorderMicrophoneAudioManager.cpp
+++ b/Engine/Plugins/VirtualProduction/Takes/Source/TakeRecorderSources/Private/TakeRecorderMicrophoneAudioManager.cpp
@@ -113,9 +113,9 @@ void UTakeRecorderMicrophoneAudioManager::FinalizeRecording()
 
 int32 UTakeRecorderMicrophoneAudioManager::GetDeviceChannelCount()
 {
-	if (AudioInputDevice.DeviceId.Len() > 0)
+	if (!AudioInputDevice.DeviceId.IsEmpty())
 	{
-		for (FAudioInputDeviceInfoProperty DeviceInfo : AudioInputDevice.DeviceInfoArray)
+		for (const FAudioInputDeviceInfoProperty& DeviceInfo : AudioInputDevice.DeviceInfoArray)
 		{
 			if (DeviceInfo.DeviceId == AudioInputDevice.DeviceId)
 			{
@@ -127,6 +127,22 @@ int32 UTakeRecorderMicrophoneAudioManager::GetDeviceChannelCount()
 	return 0;
 }
 
+bool UTakeRecorderMicrophoneAudioManager::IsAudioDeviceAvailable(const FString& InDeviceId)
+{
+	if (!InDeviceId.IsEmpty())
+	{
+		for (const FAudioInputDeviceInfoProperty& DeviceInfo : AudioInputDevice.DeviceInfoArray)
+		{
+			if (DeviceInfo.DeviceId == InDeviceId)
+			{
+				return true;
+			}
+		}
+	}
+
+	return false;
+}
+
 void UTakeRecorderMicrophoneAudioManager::BuildDeviceInfoArray()
 {
 	AudioInputDevice.DeviceInfoArray.Empty();
@@ -193,9 +209,18 @@ void UTakeRecorderMicrophoneAudioManager::BuildDeviceInfoArray()
 			return LeftLen < RightLen;
 		});
 
-	// Default to using the default input device if not already set
-	if (bFoundDefaultDevice && AudioInputDevice.DeviceId.IsEmpty())
+	if (bFoundDefaultDevice)
 	{
-		AudioInputDevice.DeviceId = DefaultDeviceInfo.DeviceId;
+		if (AudioInputDevice.DeviceId.IsEmpty())
+		{
+			// Default to using the default input device if not already set
+			AudioInputDevice.DeviceId = DefaultDeviceInfo.DeviceId;
+		}
+		else if (!IsAudioDeviceAvailable(AudioInputDevice.DeviceId))
+		{
+			// Revert to using the default input device if previously saved device is not available
+			AudioInputDevice.DeviceId = DefaultDeviceInfo.DeviceId;
+			UE_LOG(LogTakesCore, Error, TEXT("Previously saved audio input device unavailable. Falling back to default device \"%s\""), *DefaultDeviceInfo.DeviceName);
+		}
 	}
 }
diff --git a/Engine/Plugins/VirtualProduction/Takes/Source/TakeRecorderSources/Private/TakeRecorderMicrophoneAudioManager.h b/Engine/Plugins/VirtualProduction/Takes/Source/TakeRecorderSources/Private/TakeRecorderMicrophoneAudioManager.h
--- a/Engine/Plugins/VirtualProduction/Takes/Source/TakeRecorderSources/Private/TakeRecorderMicrophoneAudioManager.h
+++ b/Engine/Plugins/VirtualProduction/Takes/Source/TakeRecorderSources/Private/TakeRecorderMicrophoneAudioManager.h
@@ -70,6 +70,9 @@ private:
 	/** Builds the list of audio devices which will be used in the device menu */
 	void BuildDeviceInfoArray();
 
+	/** Returns whether audio device with given Id was found during enumeration */
+	bool IsAudioDeviceAvailable(const FString& InDeviceId);
+
 	/** The audio recorder object which manages low level recording of audio data */
 	TUniquePtr<IAudioCaptureEditor> AudioRecorder;
 };
