Sequencer: Rename take preset's director blueprint with flags

#rnx

diff --git a/Engine/Plugins/VirtualProduction/Takes/Source/TakesCore/Private/TakePreset.cpp b/Engine/Plugins/VirtualProduction/Takes/Source/TakesCore/Private/TakePreset.cpp
--- a/Engine/Plugins/VirtualProduction/Takes/Source/TakesCore/Private/TakePreset.cpp
+++ b/Engine/Plugins/VirtualProduction/Takes/Source/TakesCore/Private/TakePreset.cpp
@@ -117,7 +117,7 @@ void UTakePreset::CopyFrom(ULevelSequence* TemplateLevelSequence)
 			LevelSequence->Modify();
 
 			FName UniqueName = MakeUniqueObjectName(GetTransientPackage(), ULevelSequence::StaticClass(), "DEAD_TakePreset_LevelSequence");
-			LevelSequence->Rename(*UniqueName.ToString());
+			LevelSequence->Rename(*UniqueName.ToString(), GetTransientPackage(), REN_DoNotDirty | REN_DontCreateRedirectors | REN_ForceNoResetLoaders | REN_NonTransactional);
 		}
 
 		EObjectFlags SequenceFlags = GetFlags() & (RF_Transient | RF_Transactional);
