Take Recorder: Fix invalid level sequence on close

#rnx

diff --git a/Engine/Plugins/VirtualProduction/Takes/Source/TakeRecorder/Private/ScopedSequencerPanel.cpp b/Engine/Plugins/VirtualProduction/Takes/Source/TakeRecorder/Private/ScopedSequencerPanel.cpp
--- a/Engine/Plugins/VirtualProduction/Takes/Source/TakeRecorder/Private/ScopedSequencerPanel.cpp
+++ b/Engine/Plugins/VirtualProduction/Takes/Source/TakeRecorder/Private/ScopedSequencerPanel.cpp
@@ -114,7 +114,7 @@ void FScopedSequencerPanel::Open()
 void FScopedSequencerPanel::Close()
 {
 	ULevelSequence* LevelSequence = LevelSequenceAttribute.Get();
-	if (LevelSequence)
+	if (IsValid(LevelSequence))
 	{
 		Close(LevelSequence);
 	}
