Take Recorder: Clear take preset when asset is deleted

#rnx

diff --git a/Engine/Plugins/VirtualProduction/Takes/Source/TakeRecorder/Private/TakePresetToolkit.cpp b/Engine/Plugins/VirtualProduction/Takes/Source/TakeRecorder/Private/TakePresetToolkit.cpp
--- a/Engine/Plugins/VirtualProduction/Takes/Source/TakeRecorder/Private/TakePresetToolkit.cpp
+++ b/Engine/Plugins/VirtualProduction/Takes/Source/TakeRecorder/Private/TakePresetToolkit.cpp
@@ -58,5 +58,9 @@ FString FTakePresetToolkit::GetWorldCentricTabPrefix() const
 	return LOCTEXT("WorldCentricTabPrefix", "Take ").ToString();
 }
 
+void FTakePresetToolkit::OnClose()
+{
+	OnClosedEvent.Broadcast();
+}
 
 #undef LOCTEXT_NAMESPACE
diff --git a/Engine/Plugins/VirtualProduction/Takes/Source/TakeRecorder/Private/TakePresetToolkit.h b/Engine/Plugins/VirtualProduction/Takes/Source/TakeRecorder/Private/TakePresetToolkit.h
--- a/Engine/Plugins/VirtualProduction/Takes/Source/TakeRecorder/Private/TakePresetToolkit.h
+++ b/Engine/Plugins/VirtualProduction/Takes/Source/TakeRecorder/Private/TakePresetToolkit.h
@@ -32,6 +32,10 @@ public:
 		return GetDefaultTabIcon();
 	}
 
+	/** Called when the tab manager is changed */
+	DECLARE_EVENT(FTakePresetToolkit, FTakePresetToolkitClosed);
+	FTakePresetToolkitClosed& OnClosed() { return OnClosedEvent; }
+
 private:
 
 	virtual FText GetBaseToolkitName() const override;
@@ -40,6 +44,7 @@ private:
 	virtual FString GetWorldCentricTabPrefix() const override;
 	virtual void RegisterTabSpawners(const TSharedRef<FTabManager>& InTabManager) override {}
 	virtual void UnregisterTabSpawners(const TSharedRef<FTabManager>& InTabManager) override {}
+	virtual void OnClose() override;
 
 	virtual void AddReferencedObjects(FReferenceCollector& Collector) override;
 	virtual FString GetReferencerName() const override
@@ -53,4 +58,7 @@ private:
 	TObjectPtr<UTakePreset> TakePreset;
 
 	static const FName TabId;
+
+	/** Event that is called when this toolkit is closed */
+	FTakePresetToolkitClosed OnClosedEvent;
 };
diff --git a/Engine/Plugins/VirtualProduction/Takes/Source/TakeRecorder/Private/Widgets/STakeRecorderTabContent.cpp b/Engine/Plugins/VirtualProduction/Takes/Source/TakeRecorder/Private/Widgets/STakeRecorderTabContent.cpp
--- a/Engine/Plugins/VirtualProduction/Takes/Source/TakeRecorder/Private/Widgets/STakeRecorderTabContent.cpp
+++ b/Engine/Plugins/VirtualProduction/Takes/Source/TakeRecorder/Private/Widgets/STakeRecorderTabContent.cpp
@@ -153,6 +153,15 @@ void STakeRecorderTabContent::SetupForRecordingInto(ULevelSequence* LevelSequenc
 	}
 }
 
+void STakeRecorderTabContent::OnPresetClosed()
+{
+	WeakPanel = nullptr;
+	WeakAssetEditor = nullptr;
+
+	TitleAttribute = ITakeRecorderModule::TakeRecorderTabLabel;
+	IconAttribute = FTakeRecorderStyle::Get().GetBrush("TakeRecorder.TabIcon");
+}
+
 void STakeRecorderTabContent::SetupForEditing(TSharedPtr<FTakePresetToolkit> InToolkit)
 {
 	WeakPanel = nullptr;
@@ -160,6 +169,8 @@ void STakeRecorderTabContent::SetupForEditing(TSharedPtr<FTakePresetToolkit> InT
 	TitleAttribute = MakeAttributeSP(InToolkit.Get(), &FTakePresetToolkit::GetToolkitName);
 	IconAttribute  = MakeAttributeSP(InToolkit.Get(), &FTakePresetToolkit::GetTabIcon);
 
+	InToolkit->OnClosed().AddRaw(this, &STakeRecorderTabContent::OnPresetClosed);
+
 	// Null out the tab content to ensure that all references have been cleaned up before constructing the new one
 	ChildSlot [ SNullWidget::NullWidget ];
 
diff --git a/Engine/Plugins/VirtualProduction/Takes/Source/TakeRecorder/Private/Widgets/STakeRecorderTabContent.h b/Engine/Plugins/VirtualProduction/Takes/Source/TakeRecorder/Private/Widgets/STakeRecorderTabContent.h
--- a/Engine/Plugins/VirtualProduction/Takes/Source/TakeRecorder/Private/Widgets/STakeRecorderTabContent.h
+++ b/Engine/Plugins/VirtualProduction/Takes/Source/TakeRecorder/Private/Widgets/STakeRecorderTabContent.h
@@ -69,6 +69,8 @@ private:
 
 	EActiveTimerReturnType OnActiveTimer(double InCurrentTime, float InDeltaTime);
 
+	void OnPresetClosed();
+
 private:
 
 	TAttribute<FText> TitleAttribute;
