Append channel number to asset name only when there are name collisions.

diff --git a/Engine/Plugins/VirtualProduction/Takes/Source/TakeRecorderSources/Private/TakeRecorderMicrophoneAudioSource.cpp b/Engine/Plugins/VirtualProduction/Takes/Source/TakeRecorderSources/Private/TakeRecorderMicrophoneAudioSource.cpp
--- a/Engine/Plugins/VirtualProduction/Takes/Source/TakeRecorderSources/Private/TakeRecorderMicrophoneAudioSource.cpp
+++ b/Engine/Plugins/VirtualProduction/Takes/Source/TakeRecorderSources/Private/TakeRecorderMicrophoneAudioSource.cpp
@@ -174,20 +174,26 @@ FString UTakeRecorderMicrophoneAudioSource::GetAudioAssetName(ULevelSequence* In
 
 	if (UTakeMetaData* TakeMetaData = InSequence->FindMetaData<UTakeMetaData>())
 	{
-		bool bHasChannelCount = AudioAssetName.Contains("{channel}");
 		FString TempStr = ReplaceStringTokens(TakeMetaData->GenerateAssetPath(AudioAssetName));
 
-		if (!bHasChannelCount)
-		{
-			TempStr = FString::Printf(TEXT("%s_Channel_%d"), *TempStr, AudioChannel.AudioInputDeviceChannel);
-		}
-
 		return TempStr.Replace(TEXT(" "), TEXT("_"));
 	}
 
 	return TEXT("MicrophoneAudio");
 }
 
+FString UTakeRecorderMicrophoneAudioSource::GetUniqueAudioAssetName(ULevelSequence* InSequence) const
+{
+	FString TempStr = GetAudioAssetName(InSequence);
+
+	if (UTakeRecorderSources* SourcesList = GetTypedOuter<UTakeRecorderSources>())
+	{
+		TempStr = CreateUniqueAudioAssetName(InSequence, SourcesList, TempStr, AudioChannel.AudioInputDeviceChannel);
+	}
+
+	return TempStr;
+}
+
 FString UTakeRecorderMicrophoneAudioSource::GetSubsceneTrackName(ULevelSequence* InSequence) const
 {
 	return GetAudioTrackName(InSequence);
@@ -195,7 +201,7 @@ FString UTakeRecorderMicrophoneAudioSource::GetSubsceneTrackName(ULevelSequence*
 
 FString UTakeRecorderMicrophoneAudioSource::GetSubsceneAssetName(ULevelSequence* InSequence) const
 {
-	return GetAudioAssetName(InSequence);
+	return GetUniqueAudioAssetName(InSequence);
 }
 
 static FString MakeNewAssetName(const FString& BaseAssetPath, const FString& BaseAssetName)
@@ -268,7 +274,7 @@ TArray<UTakeRecorderSource*> UTakeRecorderMicrophoneAudioSource::PreRecording(UL
 		}
 		else
 		{
-			BaseName = GetAudioAssetName(InSequence);
+			BaseName = GetUniqueAudioAssetName(InSequence);
 		}
 
 		AudioDirectory.Path = PathToRecordTo;
@@ -277,7 +283,7 @@ TArray<UTakeRecorderSource*> UTakeRecorderMicrophoneAudioSource::PreRecording(UL
 			AudioDirectory.Path /= AudioSubDirectory;
 		}
 
-		AssetName = MakeNewAssetName(AudioDirectory.Path, BaseName);
+		AssetFileName = MakeNewAssetName(AudioDirectory.Path, BaseName);
 
 		// Add the section here so it is displayed during record (non-subsequence case)
 		UMovieSceneAudioSection* NewAudioSection = NewObject<UMovieSceneAudioSection>(CachedAudioTrack.Get(), UMovieSceneAudioSection::StaticClass());
@@ -304,27 +310,30 @@ void UTakeRecorderMicrophoneAudioSource::StartRecording(const FTimecode& InSecti
 
 	StartTimecode = InSectionStartTimecode;
 
-	const TArray<UMovieSceneSection*>& AudioSections = CachedAudioTrack->GetAudioSections();
-	if (AudioSections.Num() > 0)
-	{
-		AudioSections[0]->SetRange(TRange<FFrameNumber>(InSectionFirstFrame, InSectionFirstFrame));
-	}
-
-	UTakeRecorderMicrophoneAudioManager* AudioInputManager = GetAudioInputManager();
-	if (AudioInputManager != nullptr)
+	if (CachedAudioTrack.IsValid())
 	{
-		int32 LastChannelInUse = INDEX_NONE;
-		TUniquePtr<TArray<bool>> ChannelsInUse = GetChannelsInUse(AudioInputManager->GetDeviceChannelCount());
+		const TArray<UMovieSceneSection*>& AudioSections = CachedAudioTrack->GetAudioSections();
+		if (AudioSections.Num() > 0)
+		{
+			AudioSections[0]->SetRange(TRange<FFrameNumber>(InSectionFirstFrame, InSectionFirstFrame));
+		}
 
-		for (int32 ChannelIndex = 0; ChannelIndex < ChannelsInUse->Num(); ++ChannelIndex)
+		UTakeRecorderMicrophoneAudioManager* AudioInputManager = GetAudioInputManager();
+		if (AudioInputManager != nullptr)
 		{
-			if ((*ChannelsInUse)[ChannelIndex])
+			int32 LastChannelInUse = INDEX_NONE;
+			TUniquePtr<TArray<bool>> ChannelsInUse = GetChannelsInUse(AudioInputManager->GetDeviceChannelCount());
+
+			for (int32 ChannelIndex = 0; ChannelIndex < ChannelsInUse->Num(); ++ChannelIndex)
 			{
-				LastChannelInUse = FMath::Max(LastChannelInUse, ChannelIndex + 1);
+				if ((*ChannelsInUse)[ChannelIndex])
+				{
+					LastChannelInUse = FMath::Max(LastChannelInUse, ChannelIndex + 1);
+				}
 			}
-		}
 
-		AudioInputManager->StartRecording(LastChannelInUse);
+			AudioInputManager->StartRecording(LastChannelInUse);
+		}
 	}
 }
 
@@ -462,7 +471,7 @@ void UTakeRecorderMicrophoneAudioSource::GetRecordedSoundWave(ULevelSequence* In
 
 		FTakeRecorderAudioSourceSettings AudioSettings;
 		AudioSettings.Directory = AudioDirectory;
-		AudioSettings.AssetName = AssetName;
+		AudioSettings.AssetName = AssetFileName;
 		AudioSettings.GainDb = AudioGain;
 		AudioSettings.InputChannelNumber = AudioChannel.AudioInputDeviceChannel;
 		AudioSettings.StartTimecode = StartTimecode;
@@ -525,3 +534,31 @@ bool UTakeRecorderMicrophoneAudioSource::IsTrackAssociatedWithAnySource(UMovieSc
 
 	return false;
 }
+
+FString UTakeRecorderMicrophoneAudioSource::CreateUniqueAudioAssetName(ULevelSequence* InSequence, UTakeRecorderSources* InSources, const FString& InAssetName, const int32 InChannelNumber)
+{
+	FString NewAssetName = InAssetName;
+
+	if (ensure(InSequence) && ensure(InSources) && ensure(InChannelNumber > 0))
+	{
+		// Enforce unique name across all microphone sources
+		// Note, this enforcement is weak as it assumes channel numbers are unique. Strict uniqueness at
+		// the file name level is enforced later at asset creation time in MakeNewAssetName().
+		for (UTakeRecorderSource* Source : InSources->GetSources())
+		{
+			if (UTakeRecorderMicrophoneAudioSource* MicSource = Cast<UTakeRecorderMicrophoneAudioSource>(Source))
+			{
+				int32 ChannelNumber = MicSource->AudioChannel.AudioInputDeviceChannel;
+
+				if (ChannelNumber != InChannelNumber && NewAssetName == MicSource->GetAudioAssetName(InSequence))
+				{
+					// Append channel number to enforce uniqueness
+					NewAssetName = FString::Printf(TEXT("%s_%d"), *NewAssetName, InChannelNumber);
+					break;
+				}
+			}
+		}
+	}
+	
+	return NewAssetName;
+}
diff --git a/Engine/Plugins/VirtualProduction/Takes/Source/TakeRecorderSources/Private/TakeRecorderMicrophoneAudioSource.h b/Engine/Plugins/VirtualProduction/Takes/Source/TakeRecorderSources/Private/TakeRecorderMicrophoneAudioSource.h
--- a/Engine/Plugins/VirtualProduction/Takes/Source/TakeRecorderSources/Private/TakeRecorderMicrophoneAudioSource.h
+++ b/Engine/Plugins/VirtualProduction/Takes/Source/TakeRecorderSources/Private/TakeRecorderMicrophoneAudioSource.h
@@ -127,6 +127,12 @@ private:
 	/** Returns whether given track is being used by any other source */
 	bool IsTrackAssociatedWithAnySource(UMovieSceneAudioTrack* InAudioTrack);
 
+	/** Build asset name, appending channel name if needed for uniqueness */
+	FString GetUniqueAudioAssetName(ULevelSequence* InSequence) const;
+
+	/** Static, parameterized helper function for buidling unique asset names */	
+	static FString CreateUniqueAudioAssetName(ULevelSequence* InSequence, UTakeRecorderSources* InSources, const FString& InAssetName, const int32 InChannelNumber);
+
 private:
 
 	// Holds the Sequencer audio track for this source
@@ -139,5 +145,5 @@ private:
 	// The user specified directory to store recorded audio assets in
 	FDirectoryPath AudioDirectory;
 	// The name of the UsoundWave asset 
-	FString AssetName;
+	FString AssetFileName;
 };
