Take Recorder: Add tokens for animation asset and audio asset directories. Also, added tooltip information for what tokens can be used

diff --git a/Engine/Plugins/VirtualProduction/Takes/Source/TakeRecorderSources/Private/TakeRecorderMicrophoneAudioSource.cpp b/Engine/Plugins/VirtualProduction/Takes/Source/TakeRecorderSources/Private/TakeRecorderMicrophoneAudioSource.cpp
--- a/Engine/Plugins/VirtualProduction/Takes/Source/TakeRecorderSources/Private/TakeRecorderMicrophoneAudioSource.cpp
+++ b/Engine/Plugins/VirtualProduction/Takes/Source/TakeRecorderSources/Private/TakeRecorderMicrophoneAudioSource.cpp
@@ -280,7 +280,13 @@ TArray<UTakeRecorderSource*> UTakeRecorderMicrophoneAudioSource::PreRecording(UL
 		AudioDirectory.Path = PathToRecordTo;
 		if (AudioSubDirectory.Len())
 		{
-			AudioDirectory.Path /= AudioSubDirectory;
+			FString TempStr = AudioSubDirectory;
+			if (UTakeMetaData* TakeMetaData = InSequence->FindMetaData<UTakeMetaData>())
+			{
+				TempStr = ReplaceStringTokens(TakeMetaData->GenerateAssetPath(TempStr));
+				TempStr = TempStr.Replace(TEXT(" "), TEXT("_"));
+			}
+			AudioDirectory.Path /= TempStr;
 		}
 
 		AssetFileName = MakeNewAssetName(AudioDirectory.Path, BaseName);
diff --git a/Engine/Plugins/VirtualProduction/Takes/Source/TakeRecorderSources/Private/TakeRecorderMicrophoneAudioSource.h b/Engine/Plugins/VirtualProduction/Takes/Source/TakeRecorderSources/Private/TakeRecorderMicrophoneAudioSource.h
--- a/Engine/Plugins/VirtualProduction/Takes/Source/TakeRecorderSources/Private/TakeRecorderMicrophoneAudioSource.h
+++ b/Engine/Plugins/VirtualProduction/Takes/Source/TakeRecorderSources/Private/TakeRecorderMicrophoneAudioSource.h
@@ -41,11 +41,31 @@ public:
 	UPROPERTY(config, EditAnywhere, BlueprintReadWrite, Category = "Source")
 	FText AudioTrackName;
 
-	/** The name of the audio asset. */
+	/** The name of the audio asset.
+	 * Supports any of the following format specifiers that will be substituted when a take is recorded :
+	 * {day} - The day of the timestamp for the start of the recording.
+	 * {month} - The month of the timestamp for the start of the recording.
+	 * {year} - The year of the timestamp for the start of the recording.
+	 * {hour} - The hour of the timestamp for the start of the recording.
+	 * {minute} - The minute of the timestamp for the start of the recording.
+	 * {second} - The second of the timestamp for the start of the recording.
+	 * {take} - The take number.
+	 * {slate} - The slate string.
+	 */
 	UPROPERTY(config, EditAnywhere, BlueprintReadWrite, Category = "Source")
 	FString AudioAssetName;
 
-	/** The name of the subdirectory audio will be placed in. Leave this empty to place into the same directory as the sequence base path */
+	/** The name of the subdirectory audio will be placed in. Leave this empty to place into the same directory as the sequence base path 
+	 * Supports any of the following format specifiers that will be substituted when a take is recorded :
+	 * {day} - The day of the timestamp for the start of the recording.
+	 * {month} - The month of the timestamp for the start of the recording.
+	 * {year} - The year of the timestamp for the start of the recording.
+	 * {hour} - The hour of the timestamp for the start of the recording.
+	 * {minute} - The minute of the timestamp for the start of the recording.
+	 * {second} - The second of the timestamp for the start of the recording.
+	 * {take} - The take number.
+	 * {slate} - The slate string.
+	 */
 	UPROPERTY(config, EditAnywhere, BlueprintReadWrite, Category = "Source")
 	FString AudioSubDirectory;
 };
diff --git a/Engine/Plugins/VirtualProduction/Takes/Source/TakeTrackRecorders/Private/MovieSceneAnimationTrackRecorder.cpp b/Engine/Plugins/VirtualProduction/Takes/Source/TakeTrackRecorders/Private/MovieSceneAnimationTrackRecorder.cpp
--- a/Engine/Plugins/VirtualProduction/Takes/Source/TakeTrackRecorders/Private/MovieSceneAnimationTrackRecorder.cpp
+++ b/Engine/Plugins/VirtualProduction/Takes/Source/TakeTrackRecorders/Private/MovieSceneAnimationTrackRecorder.cpp
@@ -57,19 +57,23 @@ void UMovieSceneAnimationTrackRecorder::CreateAnimationAssetAndSequence(const AA
 		ComponentTransform = SkeletalMeshComponent->GetComponentToWorld().GetRelativeTransform(Actor->GetTransform());
 		FString AnimationAssetName = Actor->GetActorLabel();
 
+		FString DirectoryPath = AnimationDirectory.Path;
+
 		if (ULevelSequence* LeaderLevelSequence = OwningTakeRecorderSource->GetRootLevelSequence())
 		{
 			UTakeMetaData* AssetMetaData = LeaderLevelSequence->FindMetaData<UTakeMetaData>();
 
 			AnimationAssetName = AssetMetaData->GenerateAssetPath(AnimSettings->AnimationAssetName);
+			DirectoryPath = AssetMetaData->GenerateAssetPath(DirectoryPath);
 
 			TMap<FString, FStringFormatArg> FormatArgs;
 			FormatArgs.Add(TEXT("actor"), Actor->GetActorLabel());
 
 			AnimationAssetName = FString::Format(*AnimationAssetName, FormatArgs);
+			DirectoryPath = FString::Format(*DirectoryPath, FormatArgs);
 		}
 
-		AnimSequence = TakesUtils::MakeNewAsset<UAnimSequence>(AnimationDirectory.Path, AnimationAssetName);
+		AnimSequence = TakesUtils::MakeNewAsset<UAnimSequence>(DirectoryPath, AnimationAssetName);
 		if (AnimSequence.IsValid())
 		{
 			AnimSequence.Get()->MarkPackageDirty();
diff --git a/Engine/Plugins/VirtualProduction/Takes/Source/TakeTrackRecorders/Public/TrackRecorders/MovieSceneAnimationTrackRecorderSettings.h b/Engine/Plugins/VirtualProduction/Takes/Source/TakeTrackRecorders/Public/TrackRecorders/MovieSceneAnimationTrackRecorderSettings.h
--- a/Engine/Plugins/VirtualProduction/Takes/Source/TakeTrackRecorders/Public/TrackRecorders/MovieSceneAnimationTrackRecorderSettings.h
+++ b/Engine/Plugins/VirtualProduction/Takes/Source/TakeTrackRecorders/Public/TrackRecorders/MovieSceneAnimationTrackRecorderSettings.h
@@ -30,11 +30,33 @@ public:
 	UPROPERTY(config, EditAnywhere, BlueprintReadWrite, Category = "Animation Recorder Settings")
 	FText AnimationTrackName;
 
-	/** The name of the animation asset. */
+	/** The name of the animation asset. 
+	 * Supports any of the following format specifiers that will be substituted when a take is recorded:
+	 * {day}       - The day of the timestamp for the start of the recording.
+	 * {month}     - The month of the timestamp for the start of the recording.
+	 * {year}      - The year of the timestamp for the start of the recording.
+	 * {hour}      - The hour of the timestamp for the start of the recording.
+	 * {minute}    - The minute of the timestamp for the start of the recording.
+	 * {second}    - The second of the timestamp for the start of the recording.
+	 * {take}      - The take number.
+	 * {slate}     - The slate string.
+	 * {actor}     - The name of the actor being recorded.
+	 */
 	UPROPERTY(config, EditAnywhere, BlueprintReadWrite, Category = "Animation Recorder Settings")
 	FString AnimationAssetName;
 
-	/** The name of the subdirectory animations will be placed in. Leave this empty to place into the same directory as the sequence base path. */
+	/** The name of the subdirectory animations will be placed in. Leave this empty to place into the same directory as the sequence base path. 
+	 * Supports any of the following format specifiers that will be substituted when a take is recorded:
+	 * {day}       - The day of the timestamp for the start of the recording.
+	 * {month}     - The month of the timestamp for the start of the recording.
+	 * {year}      - The year of the timestamp for the start of the recording.
+	 * {hour}      - The hour of the timestamp for the start of the recording.
+	 * {minute}    - The minute of the timestamp for the start of the recording.
+	 * {second}    - The second of the timestamp for the start of the recording.
+	 * {take}      - The take number.
+	 * {slate}     - The slate string.
+	 * {actor}     - The name of the actor being recorded.
+	 */
 	UPROPERTY(config, EditAnywhere, BlueprintReadWrite, Category = "Animation Recorder Settings")
 	FString AnimationSubDirectory;
 	
