Take Recorder: Use TWeakObjectPtr to protect against components being recreated during recording

This can happen if you record a BP and then compile it during recording.

diff --git a/Engine/Plugins/VirtualProduction/Takes/Source/TakeRecorderSources/Private/TakeRecorderActorSource.cpp b/Engine/Plugins/VirtualProduction/Takes/Source/TakeRecorderSources/Private/TakeRecorderActorSource.cpp
--- a/Engine/Plugins/VirtualProduction/Takes/Source/TakeRecorderSources/Private/TakeRecorderActorSource.cpp
+++ b/Engine/Plugins/VirtualProduction/Takes/Source/TakeRecorderSources/Private/TakeRecorderActorSource.cpp
@@ -539,34 +539,34 @@ void UTakeRecorderActorSource::TickRecording(const FQualifiedFrameTime& CurrentS
 	// Each frame we want to compare against the list of components we were recording last frame. 
 	// This will allow us to detect newly added components and components that were removed at runtime,
 	// which allows us to properly update their resulting spawn track.
-	TSet<UActorComponent*> CurrentComponentList;
+	TSet<TWeakObjectPtr<UActorComponent>> CurrentComponentList;
 	GetAllComponents(CurrentComponentList, false);
 
-	TArray<UActorComponent*> NewComponentsAdded;
-	TArray<UActorComponent*> NewComponentsRemoved;
+	TArray<TWeakObjectPtr<UActorComponent>> NewComponentsAdded;
+	TArray<TWeakObjectPtr<UActorComponent>> NewComponentsRemoved;
 
-	for (UActorComponent* CurrentComponent : CurrentComponentList)
+	for (TWeakObjectPtr<UActorComponent> CurrentComponent : CurrentComponentList)
 	{
 		// Track any components added to our list this frame
-		if (!CachedComponentList.Contains(CurrentComponent))
+		if (CurrentComponent.IsValid() && !CachedComponentList.Contains(CurrentComponent))
 		{
 			NewComponentsAdded.Add(CurrentComponent);
 		}
 	}
 
-	for (UActorComponent* OldComponent : CachedComponentList)
+	for (TWeakObjectPtr<UActorComponent> OldComponent : CachedComponentList)
 	{
 		// Now do the reverse and mark any components that have been removed.
-		if (!CurrentComponentList.Contains(OldComponent))
+		if (OldComponent.IsValid() && !CurrentComponentList.Contains(OldComponent))
 		{
 			NewComponentsRemoved.Add(OldComponent);
 		}
 	}
 
 	TArray<UObject*> TraversedObjects;
-	for (UActorComponent* AddedComponent : NewComponentsAdded)
+	for (TWeakObjectPtr<UActorComponent> AddedComponent : NewComponentsAdded)
 	{
-		if (Target.IsValid() && ::IsValid(AddedComponent))
+		if (Target.IsValid() && AddedComponent.IsValid())
 		{
 			UE_LOG(LogTakesCore, Log, TEXT("Detected newly added component %s on Actor %s, beginning to record component's properties now."), *AddedComponent->GetReadableName(), *Target->GetName());
 			TSet<UMovieSceneTrackRecorder*> PreviousTrackRecorders = TSet<UMovieSceneTrackRecorder*>(TrackRecorders);
@@ -574,10 +574,10 @@ void UTakeRecorderActorSource::TickRecording(const FQualifiedFrameTime& CurrentS
 			// We should create a new property map attached to the right parent, and then initialize it using existing flow. This works for Possessables too as it will throw a warning that
 			// the binding will be broken.
 			UActorRecorderPropertyMap* ComponentPropertyMap = NewObject<UActorRecorderPropertyMap>(this, MakeUniqueObjectName(this, UActorRecorderPropertyMap::StaticClass(), AddedComponent->GetFName()), RF_Transactional);
-			ComponentPropertyMap->RecordedObject = AddedComponent;
+			ComponentPropertyMap->RecordedObject = AddedComponent.Get();
 
 			// Add the new property map as a child of the correct parent, otherwise recursion doesn't work when we try to update the cached number of recorded properties.
-			UActorRecorderPropertyMap* ParentPropertyMap = GetParentPropertyMapForComponent(AddedComponent);
+			UActorRecorderPropertyMap* ParentPropertyMap = GetParentPropertyMapForComponent(AddedComponent.Get());
 			if (ParentPropertyMap)
 			{
 				ParentPropertyMap->Children.Add(ComponentPropertyMap);
@@ -588,10 +588,10 @@ void UTakeRecorderActorSource::TickRecording(const FQualifiedFrameTime& CurrentS
 			}
 
 			// Create the Property Map
-			RebuildRecordedPropertyMapRecursive(AddedComponent, ComponentPropertyMap);
+			RebuildRecordedPropertyMapRecursive(AddedComponent.Get(), ComponentPropertyMap);
 
 			// Create the Section Recorders required
-			CreateSectionRecordersRecursive(AddedComponent, ComponentPropertyMap, TraversedObjects);
+			CreateSectionRecordersRecursive(AddedComponent.Get(), ComponentPropertyMap, TraversedObjects);
 
 			// Update our numbers on the display
 			UpdateCachedNumberOfRecordedProperties();
@@ -612,9 +612,9 @@ void UTakeRecorderActorSource::TickRecording(const FQualifiedFrameTime& CurrentS
 		}
 	}
 
-	for (UActorComponent* RemovedComponent : NewComponentsRemoved)
+	for (TWeakObjectPtr<UActorComponent> RemovedComponent : NewComponentsRemoved)
 	{
-		if (Target.IsValid() && RemovedComponent)
+		if (Target.IsValid() && RemovedComponent.IsValid())
 		{
 			UE_LOG(LogTakesCore, Log, TEXT("Detected removed component %s on Actor %s, stopping recording of component's properties now."), *RemovedComponent->GetReadableName(), *Target->GetName());
 			// sequencer-todo: notify the spawn track that no more data is needed for this without actually removing the object from the template/cdo
@@ -1252,7 +1252,7 @@ void UTakeRecorderActorSource::RebuildRecordedPropertyMapRecursive(const FFieldV
 	}
 
 	// Now try to iterate through any children on this object and continue this process recursively.
-	TSet<UActorComponent*> PossibleComponents;
+	TSet<TWeakObjectPtr<UActorComponent>> PossibleComponents;
 	TSet<AActor*> ExternalActorsReferenced;
 
 	if (InObject.IsA<AActor>())
@@ -1276,8 +1276,13 @@ void UTakeRecorderActorSource::RebuildRecordedPropertyMapRecursive(const FFieldV
 	NewReferencedActors.Append(ExternalActorsReferenced);
 
 	// Now iterate through our children and build the property map recursively.
-	for (UActorComponent* Component : PossibleComponents)
+	for (TWeakObjectPtr<UActorComponent> Component : PossibleComponents)
 	{
+		if (!Component.IsValid())
+		{
+			continue;
+		}
+
 		UE_LOG(LogTakesCore, Log, TEXT("Component: %s EditorOnly: %d Transient: %d"), *Component->GetFName().ToString(), Component->IsEditorOnly(), Component->HasAnyFlags(RF_Transient));
 		// takerecorder-todo: When merged with Dev Framework, CL 4279185, switch this to checking against
 		// IsVisualizationComponent() so that we can exclude things like default component billboards.
@@ -1289,10 +1294,10 @@ void UTakeRecorderActorSource::RebuildRecordedPropertyMapRecursive(const FFieldV
 		}
 
 		UActorRecorderPropertyMap* ComponentPropertyMap = NewObject<UActorRecorderPropertyMap>(this, MakeUniqueObjectName(this, UActorRecorderPropertyMap::StaticClass(), Component->GetFName()), RF_Transactional);
-		ComponentPropertyMap->RecordedObject = Component;
+		ComponentPropertyMap->RecordedObject = Component.Get();
 		PropertyMap->Children.Add(ComponentPropertyMap);
 
-		RebuildRecordedPropertyMapRecursive(Component, ComponentPropertyMap);
+		RebuildRecordedPropertyMapRecursive(Component.Get(), ComponentPropertyMap);
 	}
 }
 
@@ -1447,7 +1452,7 @@ void UTakeRecorderActorSource::CleanExistingDataFromSequence(const FGuid& ForGui
 	CleanExistingDataFromSequenceImpl(ForGuid, InSequence);
 }
 
-void UTakeRecorderActorSource::GetAllComponents(TSet<UActorComponent*>& OutArray, bool bUpdateReferencedActorList)
+void UTakeRecorderActorSource::GetAllComponents(TSet<TWeakObjectPtr<UActorComponent>>& OutArray, bool bUpdateReferencedActorList)
 {
 	if (Target.IsValid())
 	{
@@ -1456,7 +1461,7 @@ void UTakeRecorderActorSource::GetAllComponents(TSet<UActorComponent*>& OutArray
 	}
 }
 
-void UTakeRecorderActorSource::GetSceneComponents(USceneComponent* OnSceneComponent, TSet<UActorComponent *>& OutArray, bool bUpdateReferencedActorList)
+void UTakeRecorderActorSource::GetSceneComponents(USceneComponent* OnSceneComponent, TSet<TWeakObjectPtr<UActorComponent>>& OutArray, bool bUpdateReferencedActorList)
 {
 	if (!OnSceneComponent)
 	{
@@ -1473,16 +1478,16 @@ void UTakeRecorderActorSource::GetSceneComponents(USceneComponent* OnSceneCompon
 
 	OutArray.Add(OnSceneComponent);
 
-	TSet<UActorComponent*> ChildComponents;
+	TSet<TWeakObjectPtr<UActorComponent>> ChildComponents;
 	GetChildSceneComponents(OnSceneComponent, ChildComponents, bUpdateReferencedActorList);
 
-	for (UActorComponent* Component : ChildComponents)
+	for (TWeakObjectPtr<UActorComponent> Component : ChildComponents)
 	{
 		GetSceneComponents(Cast<USceneComponent>(Component), OutArray, bUpdateReferencedActorList);
 	}
 }
 
-void UTakeRecorderActorSource::GetChildSceneComponents(USceneComponent* OnSceneComponent, TSet<UActorComponent*>& OutArray, bool bUpdateReferencedActorList)
+void UTakeRecorderActorSource::GetChildSceneComponents(USceneComponent* OnSceneComponent, TSet<TWeakObjectPtr<UActorComponent>>& OutArray, bool bUpdateReferencedActorList)
 {
 	if (OnSceneComponent)
 	{
@@ -1562,7 +1567,7 @@ void UTakeRecorderActorSource::GetChildSceneComponents(USceneComponent* OnSceneC
 	}
 }
 
-void UTakeRecorderActorSource::GetActorComponents(AActor* OnActor, TSet<UActorComponent*>& OutArray) const
+void UTakeRecorderActorSource::GetActorComponents(AActor* OnActor, TSet<TWeakObjectPtr<UActorComponent>>& OutArray) const
 {
 	if (OnActor)
 	{
diff --git a/Engine/Plugins/VirtualProduction/Takes/Source/TakeRecorderSources/Private/TakeRecorderActorSource.h b/Engine/Plugins/VirtualProduction/Takes/Source/TakeRecorderSources/Private/TakeRecorderActorSource.h
--- a/Engine/Plugins/VirtualProduction/Takes/Source/TakeRecorderSources/Private/TakeRecorderActorSource.h
+++ b/Engine/Plugins/VirtualProduction/Takes/Source/TakeRecorderSources/Private/TakeRecorderActorSource.h
@@ -218,7 +218,7 @@ protected:
 	* @param OutArray - Set of components to add found components to.
 	* @param bUpdateReferencedActorList - If true will add any external Actors referenced to the NewReferencedActors list, and possibly print warnings.
 	*/
-	void GetAllComponents(TSet<UActorComponent*>& OutArray, bool bUpdateReferencedActorList = false);
+	void GetAllComponents(TSet<TWeakObjectPtr<UActorComponent>>& OutArray, bool bUpdateReferencedActorList = false);
 
 	/**
 	* Gets all Scene components that are a child of the specified component. Use the Root Component if you want all
@@ -227,17 +227,17 @@ protected:
 	* @param OutArray - Set of components to add found components to.
 	* @param bUpdateReferencedActorList - If true will add any external Actors referenced to the NewReferencedActors list, and possibly print warnings.
 	*/
-	void GetSceneComponents(USceneComponent* OnSceneComponent, TSet<UActorComponent*>& OutArray, bool bUpdateReferencedActorList = false);
+	void GetSceneComponents(USceneComponent* OnSceneComponent, TSet< TWeakObjectPtr<UActorComponent>>& OutArray, bool bUpdateReferencedActorList = false);
 
 	/**
 	* Returns the direct children of the specified scene component. Filters for ownership before returning.
 	*/
-	void GetChildSceneComponents(USceneComponent* OnSceneComponent, TSet<UActorComponent*>& OutArray, bool bUpdateReferencedActorList = false);
+	void GetChildSceneComponents(USceneComponent* OnSceneComponent, TSet< TWeakObjectPtr<UActorComponent>>& OutArray, bool bUpdateReferencedActorList = false);
 
 	/**
 	* Gets all Actor components on the recorded actor, not including Scene components.
 	*/
-	void GetActorComponents(AActor* OnActor, TSet<UActorComponent*>& OutArray) const;
+	void GetActorComponents(AActor* OnActor, TSet< TWeakObjectPtr<UActorComponent>>& OutArray) const;
 
 	/**
 	* Iterates through the NewReferencedActors set and creates a new UTakeRecorderActorSource for each one. Only creates
@@ -313,7 +313,7 @@ private:
 	* components have been added or removed so we can appropriately update their Spawn tracked. Empty unless a recording is
 	* in progress (but may still be empty if no components)
 	*/
-	TSet<class UActorComponent*> CachedComponentList;
+	TSet<TWeakObjectPtr<UActorComponent>> CachedComponentList;
 
 	/**
 	* ID of the TargetLevelSequence.
