Take Recorder: Sanitize package names

diff --git a/Engine/Plugins/VirtualProduction/Takes/Source/TakesCore/Public/TakesUtils.h b/Engine/Plugins/VirtualProduction/Takes/Source/TakesCore/Public/TakesUtils.h
--- a/Engine/Plugins/VirtualProduction/Takes/Source/TakesCore/Public/TakesUtils.h
+++ b/Engine/Plugins/VirtualProduction/Takes/Source/TakesCore/Public/TakesUtils.h
@@ -13,6 +13,7 @@
 #include "UObject/UObjectGlobals.h"
 #include "TakesCoreLog.h"
 #include "Misc/Paths.h"
+#include "PackageTools.h"
 
 class UWorld;
 class ULevelSequence;
@@ -119,10 +120,11 @@ namespace TakesUtils
 	static AssetType* MakeNewAsset(const FString& BaseAssetPath, const FString& BaseAssetName)
 	{
 		const FString Dot(TEXT("."));
-		FString AssetPath = BaseAssetPath;
-		FString AssetName = BaseAssetName;
-		AssetName = AssetName.Replace(TEXT("."), TEXT("_"));
-		AssetName = AssetName.Replace(TEXT(" "), TEXT("_"));
+		FString SanitizedBaseAssetPath = UPackageTools::SanitizePackageName(BaseAssetPath);
+		FString SanitizedBaseAssetName = UPackageTools::SanitizePackageName(BaseAssetName);
+
+		FString AssetPath = SanitizedBaseAssetPath;
+		FString AssetName = SanitizedBaseAssetName;
 
 		AssetName = FPaths::MakeValidFileName(AssetName);
 
@@ -144,7 +146,7 @@ namespace TakesUtils
 		}
 
 		// Create the new asset in the package we just made
-		AssetPath = (BaseAssetPath / AssetName);
+		AssetPath = (SanitizedBaseAssetPath / AssetName);
 
 		FString FileName;
 		if (FPackageName::TryConvertLongPackageNameToFilename(AssetPath, FileName))
