Re-worked audio device list menu so that it is more robust to devices going offline. Added menu item for refreshing device list which now uses FMenuBuilder to create the popup menu.

diff --git a/Engine/Plugins/Runtime/AudioCapture/Source/AudioCaptureEditor/Private/AudioRecordingManager.cpp b/Engine/Plugins/Runtime/AudioCapture/Source/AudioCaptureEditor/Private/AudioRecordingManager.cpp
--- a/Engine/Plugins/Runtime/AudioCapture/Source/AudioCaptureEditor/Private/AudioRecordingManager.cpp
+++ b/Engine/Plugins/Runtime/AudioCapture/Source/AudioCaptureEditor/Private/AudioRecordingManager.cpp
@@ -233,7 +233,7 @@ TObjectPtr<USoundWave> FAudioRecordingManager::GetRecordedSoundWave(const FRecor
 		{
 			return RecordedSoundWaves[ChannelIndex];
 		}
-		else if (NumRecordedSamples > 0)
+		else if (NumRecordedSamples > 0 && ChannelIndex < NumInputChannels)
 		{
 			return CreateSoundWaveAsset(InSourceSettings);
 		}
diff --git a/Engine/Plugins/VirtualProduction/Takes/Source/TakeRecorder/Private/Widgets/TakeRecorderAudioSettingsCustomization.cpp b/Engine/Plugins/VirtualProduction/Takes/Source/TakeRecorder/Private/Widgets/TakeRecorderAudioSettingsCustomization.cpp
--- a/Engine/Plugins/VirtualProduction/Takes/Source/TakeRecorder/Private/Widgets/TakeRecorderAudioSettingsCustomization.cpp
+++ b/Engine/Plugins/VirtualProduction/Takes/Source/TakeRecorder/Private/Widgets/TakeRecorderAudioSettingsCustomization.cpp
@@ -4,8 +4,11 @@
 #include "IDetailChildrenBuilder.h"
 #include "TakeRecorderSettings.h"
 #include "TakeRecorderSourceProperty.h"
+
+#include "Framework/MultiBox/MultiBoxBuilder.h"
 #include "Widgets/Input/SCheckBox.h"
-#include "Widgets/Input/SComboBox.h"
+
+#define LOCTEXT_NAMESPACE "FAudioInputDevicePropertyCustomization"
 
 void FAudioInputDevicePropertyCustomization::CustomizeChildren(TSharedRef<IPropertyHandle> PropertyHandle, IDetailChildrenBuilder& ChildBuilder, IPropertyTypeCustomizationUtils& CustomizationUtils)
 {
@@ -14,7 +17,7 @@ void FAudioInputDevicePropertyCustomization::CustomizeChildren(TSharedRef<IPrope
 	BufferSizeHandle = PropertyHandle->GetChildHandle(GET_MEMBER_NAME_CHECKED(FAudioInputDeviceProperty, AudioInputBufferSize));
 	DeviceInfoArrayHandle = PropertyHandle->GetChildHandle(GET_MEMBER_NAME_CHECKED(FAudioInputDeviceProperty, DeviceInfoArray))->AsArray();
 
-	// Create check box which indicates system default audio device should be used (disables device combobox)
+	// Create check box which indicates system default audio device should be used (disables device combo button)
 	if (UseSystemDefaultHandle.IsValid())
 	{
 		IDetailPropertyRow& UseSystemDefaultRow = ChildBuilder.AddProperty(PropertyHandle);
@@ -37,16 +40,21 @@ void FAudioInputDevicePropertyCustomization::CustomizeChildren(TSharedRef<IPrope
 			{
 				bool bUseSystemDefault = (NewState == ECheckBoxState::Checked);
 				UseSystemDefaultHandle->SetValue(bUseSystemDefault);
-                SynchronizeWidgetStates();
+
+				// Refresh the device list
+				RequestsDeviceListRefresh();
 			})
 		];
 	}
 
-	// Create audio input device combobox
+	// Create audio input device combo button
 	if (InputDeviceHandle.IsValid() && DeviceInfoArrayHandle.IsValid())
 	{
 		IDetailPropertyRow& InputDevicePropertyRow = ChildBuilder.AddProperty(PropertyHandle);
 
+		DeviceInfoArrayHandle->SetOnNumElementsChanged(FSimpleDelegate::CreateSP(this, &FAudioInputDevicePropertyCustomization::OnDeviceListChanged));
+		BuildAudioInputDeviceList();
+
 		InputDevicePropertyRow.CustomWidget()
 		.NameContent()
 		[
@@ -54,7 +62,7 @@ void FAudioInputDevicePropertyCustomization::CustomizeChildren(TSharedRef<IPrope
 		]
 		.ValueContent()
 		[
-			MakeAudioInputSelectorWidget()
+			CreateDeviceListComboButton()
 		];
 	}
 
@@ -73,6 +81,8 @@ void FAudioInputDevicePropertyCustomization::CustomizeChildren(TSharedRef<IPrope
 		];
 	}
 
+	SynchronizeWidgetStates();
+
 	// Register delegate so UI can update when the audio input device changes
 	// This is needed because the device menu lives in both the project settings window and the take recorder panel
 	if (UTakeRecorderAudioInputSettings* AudioInputSettings = TakeRecorderAudioSettingsUtils::GetTakeRecorderAudioInputSettings())
@@ -141,17 +151,17 @@ bool FAudioInputDevicePropertyCustomization::GetIsDefaultFromInfoProperty(const
 	return bIsDefaultDevice;
 }
 
-TSharedRef<SWidget> FAudioInputDevicePropertyCustomization::MakeAudioInputSelectorWidget()
+void FAudioInputDevicePropertyCustomization::BuildAudioInputDeviceList()
 {
 	FString CurrentDeviceId;
 	InputDeviceHandle->GetValue(CurrentDeviceId);
 
 	bool bHaveCurrentDevice = CurrentDeviceId.Len() != 0;
-	TSharedPtr<IPropertyHandle> SelectedAudioInputDevice;
-
 	uint32 NumDevices;
 	DeviceInfoArrayHandle->GetNumElements(NumDevices);
 
+	AudioInputDevices.Empty();
+
 	for (uint32 DeviceIndex = 0; DeviceIndex < NumDevices; ++DeviceIndex)
 	{
 		TSharedRef<IPropertyHandle> DeviceInfoHandle = DeviceInfoArrayHandle->GetElement(DeviceIndex);
@@ -160,11 +170,11 @@ TSharedRef<SWidget> FAudioInputDevicePropertyCustomization::MakeAudioInputSelect
 			bool bIsDefaultDevice = GetIsDefaultFromInfoProperty(DeviceInfoHandle);
 			FString DeviceId = GetDeviceIdFromInfoProperty(DeviceInfoHandle);
 
-            if (bIsDefaultDevice)
-            {
-                DefaultDeviceId = DeviceId;
-            }
-            
+			if (bIsDefaultDevice)
+			{
+				DefaultDeviceId = DeviceId;
+			}
+
 			if (bIsDefaultDevice && !bHaveCurrentDevice)
 			{
 				CurrentDeviceId = DeviceId;
@@ -182,36 +192,83 @@ TSharedRef<SWidget> FAudioInputDevicePropertyCustomization::MakeAudioInputSelect
 			AudioInputDevices.Add(DeviceInfoHandle);
 		}
 	}
+}
 
-	// Combo box component:
-	TSharedRef<SComboBox<TSharedPtr<IPropertyHandle>>> NewWidget = SNew(SComboBox<TSharedPtr<IPropertyHandle>>)
-	.OptionsSource(&AudioInputDevices)
-	.InitiallySelectedItem(SelectedAudioInputDevice)
-	.OnGenerateWidget_Lambda([](TSharedPtr<IPropertyHandle> InItem)
-	{
-		return SNew(STextBlock)
-			.Text(FText::AsCultureInvariant(GetDeviceNameFromInfoProperty(InItem)))
-			.Font(FAppStyle::Get().GetFontStyle("SmallFont"));
-	})
-	.OnSelectionChanged_Lambda([this](TSharedPtr<IPropertyHandle> InSelection, ESelectInfo::Type InSelectInfo)
+void FAudioInputDevicePropertyCustomization::RequestsDeviceListRefresh()
+{
+	TakeRecorderAudioSettingsUtils::RefreshAudioInputSettings();
+
+	BuildAudioInputDeviceList();
+	SynchronizeWidgetStates();
+}
+
+void FAudioInputDevicePropertyCustomization::OnDeviceListChanged()
+{
+	BuildAudioInputDeviceList();
+	SynchronizeWidgetStates();
+}
+
+void FAudioInputDevicePropertyCustomization::MenuItemRefreshSelected()
+{
+	RequestsDeviceListRefresh();
+}
+
+void FAudioInputDevicePropertyCustomization::MenuItemDeviceSelected(TSharedPtr<IPropertyHandle> InDeviceInfoHandle)
+{
+	SelectedAudioInputDevice = InDeviceInfoHandle;
+	InputDeviceHandle->SetValue(GetDeviceIdFromInfoProperty(InDeviceInfoHandle));
+}
+
+TSharedRef<SWidget> FAudioInputDevicePropertyCustomization::OnGenerateDeviceMenu()
+{
+	FMenuBuilder MenuBuilder(true, nullptr, nullptr);
+
+	MenuBuilder.AddMenuEntry(
+		LOCTEXT("RefreshDeviceList", "Refresh List"),
+		LOCTEXT("RefreshDeviceListToolTip", "Refresh contents of this menu's Audio Input Device List"),
+		FSlateIcon(),
+		FUIAction(
+			FExecuteAction::CreateSP(this, &FAudioInputDevicePropertyCustomization::MenuItemRefreshSelected)
+		)
+	);
+
+	MenuBuilder.BeginSection(NAME_None, LOCTEXT("AudioInputDeviceMenuSection", "Audio Input Devices"));
 	{
-		if (InSelection.IsValid() && ComboBoxTitleBlock.IsValid())
+		for (TSharedPtr<IPropertyHandle> DeviceInfoHandle : AudioInputDevices)
 		{
-			InputDeviceHandle->SetValue(GetDeviceIdFromInfoProperty(InSelection));
-			ComboBoxTitleBlock->SetText(FText::AsCultureInvariant(GetDeviceNameFromInfoProperty(InSelection)));
+			FString DeviceName = GetDeviceNameFromInfoProperty(DeviceInfoHandle);
+
+			if (!DeviceName.IsEmpty())
+			{
+				MenuBuilder.AddMenuEntry(
+					FText::FromString(DeviceName),
+					FText(),
+					FSlateIcon(),
+					FUIAction(
+						FExecuteAction::CreateSP(this, &FAudioInputDevicePropertyCustomization::MenuItemDeviceSelected, DeviceInfoHandle)
+					)
+				);
+			}
 		}
-	})
-	[
-		SAssignNew(ComboBoxTitleBlock, STextBlock)
-		.Text_Lambda([SelectedAudioInputDevice]()
+	}
+	MenuBuilder.EndSection();
+
+	return MenuBuilder.MakeWidget();
+}
+
+TSharedRef<SWidget> FAudioInputDevicePropertyCustomization::CreateDeviceListComboButton()
+{
+	TSharedRef<SSimpleComboButton> NewWidget = SNew(SSimpleComboButton)
+		.OnGetMenuContent(this, &FAudioInputDevicePropertyCustomization::OnGenerateDeviceMenu)
+		.ToolTipText(LOCTEXT("AudioInputDeviceMenuToolTip", "Available Audio Input Devices"))
+		.HasDownArrow(true)
+		.UsesSmallText(true)
+		.Text_Lambda([this]()
 		{
 			return FText::AsCultureInvariant(GetDeviceNameFromInfoProperty(SelectedAudioInputDevice));
-		})
-		.Font(FAppStyle::Get().GetFontStyle("SmallFont"))
-	];
+		});
 
-	ComboBox = NewWidget;
-	UpdateDeviceCombobxEnabled();
+	DeviceListComboButton = NewWidget;
 
 	return NewWidget;
 }
@@ -226,12 +283,12 @@ void FAudioInputDevicePropertyCustomization::UpdateDeviceCombobxEnabled()
 	bool bEnableWidget = !bUseSystemDefault;
 
 	// Get the current widget state
-	bool bIsWidgetEnabled = ComboBox->IsEnabled();
+	bool bIsWidgetEnabled = DeviceListComboButton->IsEnabled();
 
 	// If the UI doesn't match the data model, refresh the UI
 	if (bEnableWidget != bIsWidgetEnabled)
 	{
-		ComboBox->SetEnabled(bEnableWidget);
+		DeviceListComboButton->SetEnabled(bEnableWidget);
         
         // If disabling, set the current menu item to default
         if (!bEnableWidget)
@@ -252,11 +309,8 @@ void FAudioInputDevicePropertyCustomization::SynchronizeWidgetStates()
 	FString CurrentDeviceId;
 	InputDeviceHandle->GetValue(CurrentDeviceId);
 
-	// Get current menu selection
-	TSharedPtr<IPropertyHandle> CurrentMenuSelection = ComboBox->GetSelectedItem();
-	
 	// If the UI doesn't match the data model, refresh the UI
-	if (CurrentDeviceId != GetDeviceIdFromInfoProperty(CurrentMenuSelection))
+	if (CurrentDeviceId != GetDeviceIdFromInfoProperty(SelectedAudioInputDevice))
 	{
 		for (uint32 DeviceIndex = 0; DeviceIndex < NumDevices; ++DeviceIndex)
 		{
@@ -267,7 +321,8 @@ void FAudioInputDevicePropertyCustomization::SynchronizeWidgetStates()
 
 				if (DeviceId == CurrentDeviceId)
 				{
-					ComboBox->SetSelectedItem(DeviceInfoHandle);
+					SelectedAudioInputDevice = DeviceInfoHandle;
+					DeviceListComboButton->Invalidate(EInvalidateWidgetReason::Paint);
 					break;
 				}
 			}
@@ -291,3 +346,15 @@ UTakeRecorderAudioInputSettings* TakeRecorderAudioSettingsUtils::GetTakeRecorder
 	}
 	return nullptr;
 }
+
+void TakeRecorderAudioSettingsUtils::RefreshAudioInputSettings()
+{
+	// Enumerate audio devices which will rebuild the device list. Note, this can be expensive depending on the hardware
+	// attached to the machine.
+	if (UTakeRecorderAudioInputSettings* AudioInputSettings = TakeRecorderAudioSettingsUtils::GetTakeRecorderAudioInputSettings())
+	{
+		AudioInputSettings->EnumerateAudioDevices(true);
+	}
+}
+
+#undef LOCTEXT_NAMESPACE // "FAudioInputDevicePropertyCustomization"
diff --git a/Engine/Plugins/VirtualProduction/Takes/Source/TakeRecorder/Private/Widgets/TakeRecorderAudioSettingsCustomization.h b/Engine/Plugins/VirtualProduction/Takes/Source/TakeRecorder/Private/Widgets/TakeRecorderAudioSettingsCustomization.h
--- a/Engine/Plugins/VirtualProduction/Takes/Source/TakeRecorder/Private/Widgets/TakeRecorderAudioSettingsCustomization.h
+++ b/Engine/Plugins/VirtualProduction/Takes/Source/TakeRecorder/Private/Widgets/TakeRecorderAudioSettingsCustomization.h
@@ -5,7 +5,8 @@
 #include "DetailWidgetRow.h"
 #include "IPropertyTypeCustomization.h"
 #include "TakeRecorderSourceProperty.h"
-#include "Widgets/Input/SComboBox.h"
+
+#include "SSimpleComboButton.h"
 #include "Widgets/Text/STextBlock.h"
 
 class FAudioInputDevicePropertyCustomization : public IPropertyTypeCustomization
@@ -24,20 +25,26 @@ private:
 	int32 GetChannelCountFromInfoProperty(const TSharedPtr<IPropertyHandle>& InDeviceInfoHandle);
 
 	bool GetIsDefaultFromInfoProperty(const TSharedPtr<IPropertyHandle>& InDeviceInfoHandle);
-	TSharedRef<SWidget> MakeAudioInputSelectorWidget();
+	void BuildAudioInputDeviceList();
+	TSharedRef<SWidget> CreateDeviceListComboButton();
+
+	void RequestsDeviceListRefresh();
+	TSharedRef<SWidget> OnGenerateDeviceMenu();
+	void MenuItemRefreshSelected();
+	void MenuItemDeviceSelected(TSharedPtr<IPropertyHandle> InDeviceInfoHandle);
 
+	void OnDeviceListChanged();
 	void UpdateDeviceCombobxEnabled();
 	void SynchronizeWidgetStates();
 
 	TSharedPtr<IPropertyHandle> UseSystemDefaultHandle;
 	TSharedPtr<IPropertyHandle> InputDeviceHandle;
-	TSharedPtr<IPropertyHandle> DeviceChannelCountHandle;
 	TSharedPtr<IPropertyHandle> BufferSizeHandle;
 	TSharedPtr<IPropertyHandleArray> DeviceInfoArrayHandle;
 
+	TSharedPtr<IPropertyHandle> SelectedAudioInputDevice;
 	TArray<TSharedPtr<IPropertyHandle>> AudioInputDevices;
-	TSharedPtr<SComboBox<TSharedPtr<IPropertyHandle>>> ComboBox;
-	TSharedPtr<STextBlock> ComboBoxTitleBlock;
+	TSharedPtr<SSimpleComboButton> DeviceListComboButton;
 	FString DefaultDeviceId;
 };
 
@@ -46,4 +53,5 @@ class TakeRecorderAudioSettingsUtils
 public:
 
 	static UTakeRecorderAudioInputSettings* GetTakeRecorderAudioInputSettings();
+	static void RefreshAudioInputSettings();
 };
\ No newline at end of file
diff --git a/Engine/Plugins/VirtualProduction/Takes/Source/TakeRecorderSources/Private/TakeRecorderMicrophoneAudioManager.cpp b/Engine/Plugins/VirtualProduction/Takes/Source/TakeRecorderSources/Private/TakeRecorderMicrophoneAudioManager.cpp
--- a/Engine/Plugins/VirtualProduction/Takes/Source/TakeRecorderSources/Private/TakeRecorderMicrophoneAudioManager.cpp
+++ b/Engine/Plugins/VirtualProduction/Takes/Source/TakeRecorderSources/Private/TakeRecorderMicrophoneAudioManager.cpp
@@ -40,7 +40,7 @@ void UTakeRecorderMicrophoneAudioManager::PostEditChangeProperty(FPropertyChange
 
 void UTakeRecorderMicrophoneAudioManager::EnumerateAudioDevices(bool InForceRefresh)
 {
-	if (!AudioRecorder.IsValid())
+	if (!AudioRecorder.IsValid() || InForceRefresh)
 	{
 		AudioRecorder = CreateAudioRecorderObject();
 	}
