Take Recorder: Sanitize package names

diff --git a/Engine/Plugins/VirtualProduction/Takes/Source/TakeRecorderSources/Private/TakeRecorderMicrophoneAudioSource.cpp b/Engine/Plugins/VirtualProduction/Takes/Source/TakeRecorderSources/Private/TakeRecorderMicrophoneAudioSource.cpp
--- a/Engine/Plugins/VirtualProduction/Takes/Source/TakeRecorderSources/Private/TakeRecorderMicrophoneAudioSource.cpp
+++ b/Engine/Plugins/VirtualProduction/Takes/Source/TakeRecorderSources/Private/TakeRecorderMicrophoneAudioSource.cpp
@@ -23,6 +23,7 @@
 #include "Misc/PackageName.h"
 #include "ObjectEditorUtils.h"
 #include "ObjectTools.h"
+#include "PackageTools.h"
 #include "UObject/UObjectBaseUtility.h"
 
 #include UE_INLINE_GENERATED_CPP_BY_NAME(TakeRecorderMicrophoneAudioSource)
@@ -174,9 +175,7 @@ FString UTakeRecorderMicrophoneAudioSource::GetAudioAssetName(ULevelSequence* In
 
 	if (UTakeMetaData* TakeMetaData = InSequence->FindMetaData<UTakeMetaData>())
 	{
-		FString TempStr = ReplaceStringTokens(TakeMetaData->GenerateAssetPath(AudioAssetName));
-
-		return TempStr.Replace(TEXT(" "), TEXT("_"));
+		return UPackageTools::SanitizePackageName(ReplaceStringTokens(TakeMetaData->GenerateAssetPath(AudioAssetName)));
 	}
 
 	return TEXT("MicrophoneAudio");
@@ -283,8 +282,7 @@ TArray<UTakeRecorderSource*> UTakeRecorderMicrophoneAudioSource::PreRecording(UL
 			FString TempStr = AudioSubDirectory;
 			if (UTakeMetaData* TakeMetaData = InSequence->FindMetaData<UTakeMetaData>())
 			{
-				TempStr = ReplaceStringTokens(TakeMetaData->GenerateAssetPath(TempStr));
-				TempStr = TempStr.Replace(TEXT(" "), TEXT("_"));
+				TempStr = UPackageTools::SanitizePackageName(ReplaceStringTokens(TakeMetaData->GenerateAssetPath(TempStr)));
 			}
 			AudioDirectory.Path /= TempStr;
 		}
@@ -499,7 +497,6 @@ void UTakeRecorderMicrophoneAudioSource::GetRecordedSoundWave(ULevelSequence* In
 
 FText UTakeRecorderMicrophoneAudioSource::GetDisplayTextImpl() const
 {
-
 	return FText::FromString(ReplaceStringTokens(AudioSourceName.ToString()));
 }
 
