Take Recorder: Reset view and work ranges when done recording

diff --git a/Engine/Plugins/VirtualProduction/Takes/Source/TakeRecorder/Private/Recorder/TakeRecorder.cpp b/Engine/Plugins/VirtualProduction/Takes/Source/TakeRecorder/Private/Recorder/TakeRecorder.cpp
--- a/Engine/Plugins/VirtualProduction/Takes/Source/TakeRecorder/Private/Recorder/TakeRecorder.cpp
+++ b/Engine/Plugins/VirtualProduction/Takes/Source/TakeRecorder/Private/Recorder/TakeRecorder.cpp
@@ -1083,6 +1083,10 @@ void UTakeRecorder::StopInternal(const bool bCancelled)
 				MovieScene->SetPlaybackRange(CachedPlaybackRange);
 			}
 		}
+		else
+		{
+			TakesUtils::ResetViewAndWorkRange(MovieScene);
+		}
 
 		if (bRecordingFinished)
 		{
diff --git a/Engine/Plugins/VirtualProduction/Takes/Source/TakesCore/Private/TakeRecorderSources.cpp b/Engine/Plugins/VirtualProduction/Takes/Source/TakesCore/Private/TakeRecorderSources.cpp
--- a/Engine/Plugins/VirtualProduction/Takes/Source/TakesCore/Private/TakeRecorderSources.cpp
+++ b/Engine/Plugins/VirtualProduction/Takes/Source/TakesCore/Private/TakeRecorderSources.cpp
@@ -456,6 +456,7 @@ void UTakeRecorderSources::StopRecording(class ULevelSequence* InSequence, const
 			{
 				const bool bUpperBoundOnly = false; // Expand the Play Range of the sub-section to encompass all sections within it.
 				TakesUtils::ClampPlaybackRangeToEncompassAllSections(SubSequence->GetMovieScene(), bUpperBoundOnly);
+				TakesUtils::ResetViewAndWorkRange(SubSequence->GetMovieScene());
 
 				// Lock the sequence so that it can't be changed without implicitly unlocking it now
 				if (Settings.bAutoLock)
diff --git a/Engine/Plugins/VirtualProduction/Takes/Source/TakesCore/Private/TakesUtils.cpp b/Engine/Plugins/VirtualProduction/Takes/Source/TakesCore/Private/TakesUtils.cpp
--- a/Engine/Plugins/VirtualProduction/Takes/Source/TakesCore/Private/TakesUtils.cpp
+++ b/Engine/Plugins/VirtualProduction/Takes/Source/TakesCore/Private/TakesUtils.cpp
@@ -75,13 +75,18 @@ void ClampPlaybackRangeToEncompassAllSections(UMovieScene* InMovieScene, bool bU
 	}
 
 	InMovieScene->SetPlaybackRange(PlayRange.GetValue());
+}
+
+void ResetViewAndWorkRange(UMovieScene* InMovieScene)
+{
+	TRange<FFrameNumber> PlayRange = InMovieScene->GetPlaybackRange();
 
 	// Initialize the working and view range with a little bit more space
 	FFrameRate  TickResolution = InMovieScene->GetTickResolution();
-	const double OutputViewSize = PlayRange.GetValue().Size<FFrameNumber>() / TickResolution;
+	const double OutputViewSize = PlayRange.Size<FFrameNumber>() / TickResolution;
 	const double OutputChange = OutputViewSize * 0.1;
 
-	TRange<double> NewRange = UE::MovieScene::ExpandRange(PlayRange.GetValue() / TickResolution, OutputChange);
+	TRange<double> NewRange = UE::MovieScene::ExpandRange(PlayRange / TickResolution, OutputChange);
 	FMovieSceneEditorData& EditorData = InMovieScene->GetEditorData();
 	EditorData.ViewStart = EditorData.WorkStart = NewRange.GetLowerBoundValue();
 	EditorData.ViewEnd = EditorData.WorkEnd = NewRange.GetUpperBoundValue();
diff --git a/Engine/Plugins/VirtualProduction/Takes/Source/TakesCore/Public/TakesUtils.h b/Engine/Plugins/VirtualProduction/Takes/Source/TakesCore/Public/TakesUtils.h
--- a/Engine/Plugins/VirtualProduction/Takes/Source/TakesCore/Public/TakesUtils.h
+++ b/Engine/Plugins/VirtualProduction/Takes/Source/TakesCore/Public/TakesUtils.h
@@ -28,6 +28,8 @@ namespace TakesUtils
 
 	TAKESCORE_API void ClampPlaybackRangeToEncompassAllSections(UMovieScene* InMovieScene, bool bUpperBoundOnly);
 
+	TAKESCORE_API void ResetViewAndWorkRange(UMovieScene* InMovieScene);
+
 	TAKESCORE_API void SaveAsset(UObject* InObject);
 
 	TAKESCORE_API void CreateCameraCutTrack(ULevelSequence* LevelSequence, const FGuid& RecordedCameraGuid, const FMovieSceneSequenceID& SequenceID, const TRange<FFrameNumber>& InRange);
