Take Recorder: Remove references to Private directory of TakeRecorderSources

#rnx

diff --git a/Engine/Plugins/VirtualProduction/Takes/Source/TakeRecorder/TakeRecorder.Build.cs b/Engine/Plugins/VirtualProduction/Takes/Source/TakeRecorder/TakeRecorder.Build.cs
--- a/Engine/Plugins/VirtualProduction/Takes/Source/TakeRecorder/TakeRecorder.Build.cs
+++ b/Engine/Plugins/VirtualProduction/Takes/Source/TakeRecorder/TakeRecorder.Build.cs
@@ -45,12 +45,6 @@ public class TakeRecorder : ModuleRules
                 "SerializedRecorderInterface",
                 "Sequencer",
             }
-        );
-
-		PrivateIncludePaths.AddRange(
-			new string[] {
-                "TakeRecorderSources/Private",
-            }
         );
     }
 }
