Take Recorder: Change PreInitialize, Sarted, Stopped, Finished, Cancelled, MarkedFrameAdded to multicast delegates

Deprecate delegates that the multicast delegates replace

diff --git a/Engine/Plugins/VirtualProduction/Takes/Source/TakeRecorder/Private/Recorder/TakeRecorder.cpp b/Engine/Plugins/VirtualProduction/Takes/Source/TakeRecorder/Private/Recorder/TakeRecorder.cpp
--- a/Engine/Plugins/VirtualProduction/Takes/Source/TakeRecorder/Private/Recorder/TakeRecorder.cpp
+++ b/Engine/Plugins/VirtualProduction/Takes/Source/TakeRecorder/Private/Recorder/TakeRecorder.cpp
@@ -16,6 +16,7 @@
 #include "MovieSceneTimeHelpers.h"
 #include "ObjectTools.h"
 #include "Recorder/TakeRecorderBlueprintLibrary.h"
+#include "Recorder/TakeRecorderSubsystem.h"
 #include "TakeMetaData.h"
 #include "TakePreset.h"
 #include "TakeRecorderOverlayWidget.h"
@@ -453,6 +454,10 @@ bool UTakeRecorder::Initialize ( ULevelSequence* LevelSequenceBase, UTakeRecorde
 	OnRecordingPreInitializeEvent.Broadcast(this);
 
 	UTakeRecorderBlueprintLibrary::OnTakeRecorderPreInitialize();
+	if (UTakeRecorderSubsystem* TakeRecorderSubsystem = GEngine->GetEngineSubsystem<UTakeRecorderSubsystem>())
+	{
+		TakeRecorderSubsystem->TakeRecorderPreInitialize.Broadcast();
+	}
 
 	FTakeRecorderParameters FinalParameters = TakeInitHelper::AccumulateParamsOverride(InParameters);
 	if (FinalParameters.TakeRecorderMode == ETakeRecorderMode::RecordNewSequence)
@@ -969,6 +974,10 @@ void UTakeRecorder::Start()
 	OnRecordingStartedEvent.Broadcast(this);
 
 	UTakeRecorderBlueprintLibrary::OnTakeRecorderStarted();
+	if (UTakeRecorderSubsystem* TakeRecorderSubsystem = GEngine->GetEngineSubsystem<UTakeRecorderSubsystem>())
+	{
+		TakeRecorderSubsystem->TakeRecorderStarted.Broadcast();
+	}
 }
 
 void UTakeRecorder::Stop()
@@ -1046,7 +1055,12 @@ void UTakeRecorder::StopInternal(const bool bCancelled)
 		}
 
 		OnRecordingStoppedEvent.Broadcast(this);
+
 		UTakeRecorderBlueprintLibrary::OnTakeRecorderStopped();
+		if (UTakeRecorderSubsystem* TakeRecorderSubsystem = GEngine->GetEngineSubsystem<UTakeRecorderSubsystem>())
+		{
+			TakeRecorderSubsystem->TakeRecorderStopped.Broadcast();
+		}
 
 		if (MovieScene)
 		{
@@ -1146,11 +1160,19 @@ void UTakeRecorder::StopInternal(const bool bCancelled)
 		{
 			OnRecordingFinishedEvent.Broadcast(this);
 			UTakeRecorderBlueprintLibrary::OnTakeRecorderFinished(SequenceAsset);
+			if (UTakeRecorderSubsystem* TakeRecorderSubsystem = GEngine->GetEngineSubsystem<UTakeRecorderSubsystem>())
+			{
+				TakeRecorderSubsystem->TakeRecorderFinished.Broadcast(SequenceAsset);
+			}
 		}
 		else
 		{
 			OnRecordingCancelledEvent.Broadcast(this);
 			UTakeRecorderBlueprintLibrary::OnTakeRecorderCancelled();
+			if (UTakeRecorderSubsystem* TakeRecorderSubsystem = GEngine->GetEngineSubsystem<UTakeRecorderSubsystem>())
+			{
+				TakeRecorderSubsystem->TakeRecorderCancelled.Broadcast();
+			}
 		}
 	}
 
diff --git a/Engine/Plugins/VirtualProduction/Takes/Source/TakeRecorder/Private/Recorder/TakeRecorderBlueprintLibrary.cpp b/Engine/Plugins/VirtualProduction/Takes/Source/TakeRecorder/Private/Recorder/TakeRecorderBlueprintLibrary.cpp
--- a/Engine/Plugins/VirtualProduction/Takes/Source/TakeRecorder/Private/Recorder/TakeRecorderBlueprintLibrary.cpp
+++ b/Engine/Plugins/VirtualProduction/Takes/Source/TakeRecorder/Private/Recorder/TakeRecorderBlueprintLibrary.cpp
@@ -13,7 +13,7 @@
 
 #include UE_INLINE_GENERATED_CPP_BY_NAME(TakeRecorderBlueprintLibrary)
 
-namespace
+namespace TakeRecorderBlueprintLibrary
 {
 	static UTakeRecorderPanel* CurrentTakeRecorderPanel;
 	static FOnTakeRecorderPanelChanged TakeRecorderPanelChanged;
@@ -141,93 +141,93 @@ UTakeRecorderPanel* UTakeRecorderBlueprintLibrary::OpenTakeRecorderPanel()
 
 UTakeRecorderPanel* UTakeRecorderBlueprintLibrary::GetTakeRecorderPanel()
 {
-	return CurrentTakeRecorderPanel;
+	return TakeRecorderBlueprintLibrary::CurrentTakeRecorderPanel;
 }
 
 void UTakeRecorderBlueprintLibrary::SetOnTakeRecorderPanelChanged(FOnTakeRecorderPanelChanged OnTakeRecorderPanelChanged)
 {
-	TakeRecorderPanelChanged = OnTakeRecorderPanelChanged;
+	TakeRecorderBlueprintLibrary::TakeRecorderPanelChanged = OnTakeRecorderPanelChanged;
 };
 
 void UTakeRecorderBlueprintLibrary::SetOnTakeRecorderPreInitialize(FOnTakeRecorderPreInitialize OnTakeRecorderPreInitialize)
 {
-	TakeRecorderPreInitialize = OnTakeRecorderPreInitialize;
+	TakeRecorderBlueprintLibrary::TakeRecorderPreInitialize = OnTakeRecorderPreInitialize;
 };
 
 void UTakeRecorderBlueprintLibrary::SetOnTakeRecorderStarted(FOnTakeRecorderStarted OnTakeRecorderStarted)
 {
-	TakeRecorderStarted = OnTakeRecorderStarted;
+	TakeRecorderBlueprintLibrary::TakeRecorderStarted = OnTakeRecorderStarted;
 };
 
 void UTakeRecorderBlueprintLibrary::SetOnTakeRecorderStopped(FOnTakeRecorderStopped OnTakeRecorderStopped)
 {
-	TakeRecorderStopped = OnTakeRecorderStopped;
+	TakeRecorderBlueprintLibrary::TakeRecorderStopped = OnTakeRecorderStopped;
 };
 
 void UTakeRecorderBlueprintLibrary::SetOnTakeRecorderFinished(FOnTakeRecorderFinished OnTakeRecorderFinished)
 {
-	TakeRecorderFinished = OnTakeRecorderFinished;
+	TakeRecorderBlueprintLibrary::TakeRecorderFinished = OnTakeRecorderFinished;
 };
 
 void UTakeRecorderBlueprintLibrary::SetOnTakeRecorderCancelled(FOnTakeRecorderCancelled OnTakeRecorderCancelled)
 {
-	TakeRecorderCancelled = OnTakeRecorderCancelled;
+	TakeRecorderBlueprintLibrary::TakeRecorderCancelled = OnTakeRecorderCancelled;
 };
 
 void UTakeRecorderBlueprintLibrary::SetOnTakeRecorderMarkedFrameAdded(FOnTakeRecorderMarkedFrameAdded OnTakeRecorderMarkedFrameAdded)
 {
-	TakeRecorderMarkedFrameAdded = OnTakeRecorderMarkedFrameAdded;
+	TakeRecorderBlueprintLibrary::TakeRecorderMarkedFrameAdded = OnTakeRecorderMarkedFrameAdded;
 };
 
 void UTakeRecorderBlueprintLibrary::OnTakeRecorderPreInitialize()
 {
-	TakeRecorderPreInitialize.ExecuteIfBound();
+	TakeRecorderBlueprintLibrary::TakeRecorderPreInitialize.ExecuteIfBound();
 }
 
 void UTakeRecorderBlueprintLibrary::OnTakeRecorderStarted()
 {
-	TakeRecorderStarted.ExecuteIfBound();
+	TakeRecorderBlueprintLibrary::TakeRecorderStarted.ExecuteIfBound();
 }
 
 void UTakeRecorderBlueprintLibrary::OnTakeRecorderStopped()
 {
-	TakeRecorderStopped.ExecuteIfBound();
+	TakeRecorderBlueprintLibrary::TakeRecorderStopped.ExecuteIfBound();
 }
 
 void UTakeRecorderBlueprintLibrary::OnTakeRecorderFinished(ULevelSequence* InSequenceAsset)
 {
-	TakeRecorderFinished.ExecuteIfBound(InSequenceAsset);
+	TakeRecorderBlueprintLibrary::TakeRecorderFinished.ExecuteIfBound(InSequenceAsset);
 }
 
 void UTakeRecorderBlueprintLibrary::OnTakeRecorderCancelled()
 {
-	TakeRecorderCancelled.ExecuteIfBound();
+	TakeRecorderBlueprintLibrary::TakeRecorderCancelled.ExecuteIfBound();
 }
 
 void UTakeRecorderBlueprintLibrary::OnTakeRecorderMarkedFrameAdded(const FMovieSceneMarkedFrame& InMarkedFrame)
 {
-	TakeRecorderMarkedFrameAdded.ExecuteIfBound(InMarkedFrame);
+	TakeRecorderBlueprintLibrary::TakeRecorderMarkedFrameAdded.ExecuteIfBound(InMarkedFrame);
 }
 
 void UTakeRecorderBlueprintLibrary::SetTakeRecorderPanel(UTakeRecorderPanel* InNewPanel)
 {
-	if (CurrentTakeRecorderPanel != InNewPanel)
+	if (TakeRecorderBlueprintLibrary::CurrentTakeRecorderPanel != InNewPanel)
 	{
-		if (CurrentTakeRecorderPanel)
+		if (TakeRecorderBlueprintLibrary::CurrentTakeRecorderPanel)
 		{
 			// Old panel is no longer valid
-			CurrentTakeRecorderPanel->ClosePanel();
+			TakeRecorderBlueprintLibrary::CurrentTakeRecorderPanel->ClosePanel();
 
-			CurrentTakeRecorderPanel->RemoveFromRoot();
-			CurrentTakeRecorderPanel = nullptr;
+			TakeRecorderBlueprintLibrary::CurrentTakeRecorderPanel->RemoveFromRoot();
+			TakeRecorderBlueprintLibrary::CurrentTakeRecorderPanel = nullptr;
 		}
 
 		if (InNewPanel && InNewPanel->IsPanelOpen())
 		{
 			InNewPanel->AddToRoot();
-			CurrentTakeRecorderPanel = InNewPanel;
+			TakeRecorderBlueprintLibrary::CurrentTakeRecorderPanel = InNewPanel;
 		}
 
-		TakeRecorderPanelChanged.ExecuteIfBound();
+		TakeRecorderBlueprintLibrary::TakeRecorderPanelChanged.ExecuteIfBound();
 	}
 }
diff --git a/Engine/Plugins/VirtualProduction/Takes/Source/TakeRecorder/Private/Recorder/TakeRecorderSubsystem.cpp b/Engine/Plugins/VirtualProduction/Takes/Source/TakeRecorder/Private/Recorder/TakeRecorderSubsystem.cpp
new file mode 100644
--- /dev/null
+++ b/Engine/Plugins/VirtualProduction/Takes/Source/TakeRecorder/Private/Recorder/TakeRecorderSubsystem.cpp
@@ -0,0 +1,7 @@
+// Copyright Epic Games, Inc. All Rights Reserved.
+
+#include "Recorder/TakeRecorderSubsystem.h"
+
+#define LOCTEXT_NAMESPACE "TakeRecorder"
+
+#undef LOCTEXT_NAMESPACE
diff --git a/Engine/Plugins/VirtualProduction/Takes/Source/TakeRecorder/Private/Widgets/STakeRecorderCockpit.cpp b/Engine/Plugins/VirtualProduction/Takes/Source/TakeRecorder/Private/Widgets/STakeRecorderCockpit.cpp
--- a/Engine/Plugins/VirtualProduction/Takes/Source/TakeRecorder/Private/Widgets/STakeRecorderCockpit.cpp
+++ b/Engine/Plugins/VirtualProduction/Takes/Source/TakeRecorder/Private/Widgets/STakeRecorderCockpit.cpp
@@ -17,6 +17,7 @@
 #include "MovieScene.h"
 #include "Recorder/TakeRecorder.h"
 #include "Recorder/TakeRecorderBlueprintLibrary.h"
+#include "Recorder/TakeRecorderSubsystem.h"
 #include "LevelSequence.h"
 #include "Algo/Find.h"
 #include "Framework/MultiBox/MultiBoxBuilder.h"
@@ -1005,6 +1006,10 @@ FReply STakeRecorderCockpit::OnAddMarkedFrame()
 		
 		int32 MarkedFrameIndex = MovieScene->AddMarkedFrame(MarkedFrame);
 		UTakeRecorderBlueprintLibrary::OnTakeRecorderMarkedFrameAdded(MovieScene->GetMarkedFrames()[MarkedFrameIndex]);
+		if (UTakeRecorderSubsystem* TakeRecorderSubsystem = GEngine->GetEngineSubsystem<UTakeRecorderSubsystem>())
+		{
+			TakeRecorderSubsystem->TakeRecorderMarkedFrameAdded.Broadcast(MovieScene->GetMarkedFrames()[MarkedFrameIndex]);
+		}
 	}
 
 	return FReply::Handled();
diff --git a/Engine/Plugins/VirtualProduction/Takes/Source/TakeRecorder/Public/Recorder/TakeRecorderBlueprintLibrary.h b/Engine/Plugins/VirtualProduction/Takes/Source/TakeRecorder/Public/Recorder/TakeRecorderBlueprintLibrary.h
--- a/Engine/Plugins/VirtualProduction/Takes/Source/TakeRecorder/Public/Recorder/TakeRecorderBlueprintLibrary.h
+++ b/Engine/Plugins/VirtualProduction/Takes/Source/TakeRecorder/Public/Recorder/TakeRecorderBlueprintLibrary.h
@@ -105,28 +105,28 @@ public:
 	UFUNCTION(BlueprintCallable, Category="Take Recorder", meta=(DisplayName="Set On Take Recorder Panel Changed"))
 	static void SetOnTakeRecorderPanelChanged(FOnTakeRecorderPanelChanged OnTakeRecorderPanelChanged);
 
-	/** Called before initialization occurs (ie. when the recording button is pressed and before the countdown starts) */
-	UFUNCTION(BlueprintCallable, Category = "Take Recorder")
+	UE_DEPRECATED(5.4, "Please use TakeRecorderSubsystem::TakeRecorderPreInitialize")
+	UFUNCTION(BlueprintCallable, Category = "Take Recorder", meta = (DeprecatedFunction, DeprecationMessage = "Please use TakeRecorderSubsystem::TakeRecorderPreInitialize"))
 	static void SetOnTakeRecorderPreInitialize(FOnTakeRecorderPreInitialize OnTakeRecorderPreInitialize);
 
-	/** Called when take recording starts. */
-	UFUNCTION(BlueprintCallable, Category = "Take Recorder")
+	UE_DEPRECATED(5.4, "Please use TakeRecorderSubsystem::TakeRecorderStarted")
+	UFUNCTION(BlueprintCallable, Category = "Take Recorder", meta = (DeprecatedFunction, DeprecationMessage = "Please use TakeRecorderSubsystem::TakeRecorderStarted"))
 	static void SetOnTakeRecorderStarted(FOnTakeRecorderStarted OnTakeRecorderStarted);
 
-	/** Called when take recording is stopped. */
-	UFUNCTION(BlueprintCallable, Category = "Take Recorder")
+	UE_DEPRECATED(5.4, "Please use TakeRecorderSubsystem::TakeRecorderStopped")
+	UFUNCTION(BlueprintCallable, Category = "Take Recorder", meta = (DeprecatedFunction, DeprecationMessage = "Please use TakeRecorderSubsystem::TakeRecorderStopped"))
 	static void SetOnTakeRecorderStopped(FOnTakeRecorderStopped OnTakeRecorderStopped);
 
-	/** Called when take recording finishes. */
-	UFUNCTION(BlueprintCallable, Category = "Take Recorder")
+	UE_DEPRECATED(5.4, "Please use TakeRecorderSubsystem::TakeRecorderFinished")
+	UFUNCTION(BlueprintCallable, Category = "Take Recorder", meta = (DeprecatedFunction, DeprecationMessage = "Please use TakeRecorderSubsystem::TakeRecorderFinished"))
 	static void SetOnTakeRecorderFinished(FOnTakeRecorderFinished OnTakeRecorderFinished);
 
-	/** Called when take recording is cancelled. */
-	UFUNCTION(BlueprintCallable, Category = "Take Recorder")
+	UE_DEPRECATED(5.4, "Please use TakeRecorderSubsystem::TakeRecorderCancelled")
+	UFUNCTION(BlueprintCallable, Category = "Take Recorder", meta = (DeprecatedFunction, DeprecationMessage = "Please use TakeRecorderSubsystem::TakeRecorderCancelled"))
 	static void SetOnTakeRecorderCancelled(FOnTakeRecorderCancelled OnTakeRecorderCancelled);
 
-	/** Called when a marked frame is added. */
-	UFUNCTION(BlueprintCallable, Category = "Take Recorder")
+	UE_DEPRECATED(5.4, "Please use TakeRecorderSubsystem::TakeRecorderMarkedFrameAdded")
+	UFUNCTION(BlueprintCallable, Category = "Take Recorder", meta = (DeprecatedFunction, DeprecationMessage = "Please use TakeRecorderSubsystem::TakeRecorderMarkedFrameAdded"))
 	static void SetOnTakeRecorderMarkedFrameAdded(FOnTakeRecorderMarkedFrameAdded OnTakeRecorderMarkedFrameAdded);
 
 	static void OnTakeRecorderPreInitialize();
diff --git a/Engine/Plugins/VirtualProduction/Takes/Source/TakeRecorder/Public/Recorder/TakeRecorderSubsystem.h b/Engine/Plugins/VirtualProduction/Takes/Source/TakeRecorder/Public/Recorder/TakeRecorderSubsystem.h
new file mode 100644
--- /dev/null
+++ b/Engine/Plugins/VirtualProduction/Takes/Source/TakeRecorder/Public/Recorder/TakeRecorderSubsystem.h
@@ -0,0 +1,52 @@
+// Copyright Epic Games, Inc. All Rights Reserved.
+
+#pragma once
+
+#include "Subsystems/EngineSubsystem.h"
+
+#include "TakeRecorderSubsystem.generated.h"
+
+class ULevelSequence;
+
+DECLARE_DYNAMIC_MULTICAST_DELEGATE(FTakeRecorderPreInitialize);
+DECLARE_DYNAMIC_MULTICAST_DELEGATE(FTakeRecorderStarted);
+DECLARE_DYNAMIC_MULTICAST_DELEGATE(FTakeRecorderStopped);
+DECLARE_DYNAMIC_MULTICAST_DELEGATE_OneParam(FTakeRecorderFinished, ULevelSequence*, SequenceAsset);
+DECLARE_DYNAMIC_MULTICAST_DELEGATE(FTakeRecorderCancelled);
+DECLARE_DYNAMIC_MULTICAST_DELEGATE_OneParam(FTakeRecorderMarkedFrameAdded, const FMovieSceneMarkedFrame&, MarkedFrame);
+
+/**
+* UTakeRecorderSubsystem
+* Subsystem for Take Recorder
+*/
+UCLASS()
+class TAKERECORDER_API UTakeRecorderSubsystem : public UEngineSubsystem
+{
+	GENERATED_BODY()
+
+public:
+
+	/** Called before initialization occurs (ie. when the recording button is pressed and before the countdown starts) */
+	UPROPERTY(BlueprintAssignable, Category = "Take Recorder")
+	FTakeRecorderPreInitialize TakeRecorderPreInitialize;
+
+	/** Called when take recorder is started */
+	UPROPERTY(BlueprintAssignable, Category = "Take Recorder")
+	FTakeRecorderStarted TakeRecorderStarted;
+
+	/** Called when take recorder is stopped */
+	UPROPERTY(BlueprintAssignable, Category = "Take Recorder")
+	FTakeRecorderStopped TakeRecorderStopped;
+
+	/** Called when take recorder has finished */
+	UPROPERTY(BlueprintAssignable, Category = "Take Recorder")
+	FTakeRecorderFinished TakeRecorderFinished;
+
+	/** Called when take recorder is cancelled */
+	UPROPERTY(BlueprintAssignable, Category = "Take Recorder")
+	FTakeRecorderCancelled TakeRecorderCancelled;
+
+	/** Called when a marked frame is added to take recorder */
+	UPROPERTY(BlueprintAssignable, Category = "Take Recorder")
+	FTakeRecorderMarkedFrameAdded TakeRecorderMarkedFrameAdded;
+};
\ No newline at end of file
diff --git a/Engine/Plugins/VirtualProduction/Takes/Source/TakeRecorder/TakeRecorder.Build.cs b/Engine/Plugins/VirtualProduction/Takes/Source/TakeRecorder/TakeRecorder.Build.cs
--- a/Engine/Plugins/VirtualProduction/Takes/Source/TakeRecorder/TakeRecorder.Build.cs
+++ b/Engine/Plugins/VirtualProduction/Takes/Source/TakeRecorder/TakeRecorder.Build.cs
@@ -40,6 +40,7 @@ public class TakeRecorder : ModuleRules
 
 		PublicDependencyModuleNames.AddRange(
 			new string[] {
+				"EditorSubsystem",
 				"UMG",
                 "TakeTrackRecorders",
                 "SerializedRecorderInterface",
