Take Recorder: Set the animation track name

diff --git a/Engine/Plugins/VirtualProduction/Takes/Source/TakeTrackRecorders/Private/MovieSceneAnimationTrackRecorder.cpp b/Engine/Plugins/VirtualProduction/Takes/Source/TakeTrackRecorders/Private/MovieSceneAnimationTrackRecorder.cpp
--- a/Engine/Plugins/VirtualProduction/Takes/Source/TakeTrackRecorders/Private/MovieSceneAnimationTrackRecorder.cpp
+++ b/Engine/Plugins/VirtualProduction/Takes/Source/TakeTrackRecorders/Private/MovieSceneAnimationTrackRecorder.cpp
@@ -199,6 +199,7 @@ void UMovieSceneAnimationTrackRecorder::CreateTrackImpl()
 
 			if (AnimTrack)
 			{
+				AnimTrack->SetDisplayName(AnimSettings->AnimationTrackName);
 				AnimTrack->AddNewAnimation(FFrameNumber(0), AnimSequence.Get());
 				MovieSceneSection = Cast<UMovieSceneSkeletalAnimationSection>(AnimTrack->GetAllSections()[0]);
 				MovieSceneSection->Params.bForceCustomMode = true;
