Sequencer: conditional recompilation suppressor

diff --git a/Engine/Plugins/VirtualProduction/Takes/Source/TakeRecorder/Private/Recorder/TakeRecorder.cpp b/Engine/Plugins/VirtualProduction/Takes/Source/TakeRecorder/Private/Recorder/TakeRecorder.cpp
--- a/Engine/Plugins/VirtualProduction/Takes/Source/TakeRecorder/Private/Recorder/TakeRecorder.cpp
+++ b/Engine/Plugins/VirtualProduction/Takes/Source/TakeRecorder/Private/Recorder/TakeRecorder.cpp
@@ -8,6 +8,9 @@
 #include "CoreGlobals.h"
 #include "Engine/Engine.h"
 #include "Engine/EngineTypes.h"
+#include "EntitySystem/MovieSceneEntitySystemLinker.h"
+#include "EntitySystem/MovieSceneInstanceRegistry.h"
+#include "Evaluation/MovieSceneEvaluationTemplateInstance.h"
 #include "AssetRegistry/IAssetRegistry.h"
 #include "ISequencer.h"
 #include "LevelSequence.h"
@@ -751,7 +754,9 @@ FQualifiedFrameTime UTakeRecorder::GetRecordTime() const
 
 void UTakeRecorder::InternalTick(float DeltaTime)
 {
-	UE::MovieScene::FScopedSignedObjectModifyDefer FlushOnTick(true);
+	using namespace UE::MovieScene;
+
+	FScopedSignedObjectModifyDefer FlushOnTick(true);
 
 	TSharedPtr<ISequencer> Sequencer = WeakSequencer.Pin();
 	FQualifiedFrameTime RecordTime = GetRecordTime();
@@ -892,6 +897,8 @@ void UTakeRecorder::PreRecord()
 
 void UTakeRecorder::Start()
 {
+	using namespace UE::MovieScene;
+
 	FTimecode Timecode = FApp::GetTimecode();
 
 	State = ETakeRecorderState::Started;
@@ -978,6 +985,10 @@ void UTakeRecorder::Start()
 	{
 		TakeRecorderSubsystem->TakeRecorderStarted.Broadcast();
 	}
+
+	FRootInstanceHandle RootInstanceHandle = Sequencer->GetEvaluationTemplate().GetRootInstanceHandle();
+	FInstanceRegistry* InstanceRegistry = Sequencer->GetEvaluationTemplate().GetEntitySystemLinker()->GetInstanceRegistry();
+	CompileSuppression = MakeUnique<FScopedVolatilityManagerSuppression>(InstanceRegistry, RootInstanceHandle);
 }
 
 void UTakeRecorder::Stop()
@@ -1039,6 +1050,8 @@ void UTakeRecorder::StopInternal(const bool bCancelled)
 		}
 	}
 
+	CompileSuppression.Reset();
+
 	if (bDidEverStartRecording)
 	{
 		if (!ShouldShowNotifications())
diff --git a/Engine/Plugins/VirtualProduction/Takes/Source/TakeRecorder/Public/Recorder/TakeRecorder.h b/Engine/Plugins/VirtualProduction/Takes/Source/TakeRecorder/Public/Recorder/TakeRecorder.h
--- a/Engine/Plugins/VirtualProduction/Takes/Source/TakeRecorder/Public/Recorder/TakeRecorder.h
+++ b/Engine/Plugins/VirtualProduction/Takes/Source/TakeRecorder/Public/Recorder/TakeRecorder.h
@@ -19,6 +19,10 @@ class UTakeRecorderSources;
 class UTakeRecorderOverlayWidget;
 class UMovieSceneSequence;
 
+namespace UE::MovieScene
+{
+	struct FScopedVolatilityManagerSuppression;
+}
 
 UENUM(BlueprintType)
 enum class ETakeRecorderState : uint8
@@ -310,4 +314,6 @@ private:
 	EUpdateClockSource CachedClockSource;
 	
 	TRange<FFrameNumber> CachedPlaybackRange;
+
+	TUniquePtr<UE::MovieScene::FScopedVolatilityManagerSuppression> CompileSuppression;
 };
diff --git a/Engine/Source/Runtime/MovieScene/Private/EntitySystem/MovieSceneInstanceRegistry.cpp b/Engine/Source/Runtime/MovieScene/Private/EntitySystem/MovieSceneInstanceRegistry.cpp
--- a/Engine/Source/Runtime/MovieScene/Private/EntitySystem/MovieSceneInstanceRegistry.cpp
+++ b/Engine/Source/Runtime/MovieScene/Private/EntitySystem/MovieSceneInstanceRegistry.cpp
@@ -11,6 +11,7 @@
 #include "EntitySystem/BuiltInComponentTypes.h"
 
 #include "Compilation/MovieSceneCompiledDataManager.h"
+#include "Compilation/MovieSceneCompiledVolatilityManager.h"
 
 #include "Evaluation/MovieSceneSequenceHierarchy.h"
 #include "Evaluation/MovieSceneEvaluationTemplateInstance.h"
@@ -127,5 +128,20 @@ void FInstanceRegistry::CleanupLinkerEntities(const TSet<FMovieSceneEntityID>& E
 	}
 }
 
+FScopedVolatilityManagerSuppression::FScopedVolatilityManagerSuppression(FInstanceRegistry* InInstanceRegistry, FRootInstanceHandle InRootInstanceHandle)
+	: InstanceRegistry(InInstanceRegistry)
+	, RootInstanceHandle(InRootInstanceHandle)
+{
+	FSequenceInstance& Instance = InstanceRegistry->MutateInstance(RootInstanceHandle);
+	PreviousVolatilityManager = MoveTemp(Instance.VolatilityManager);
+}
+
+FScopedVolatilityManagerSuppression::~FScopedVolatilityManagerSuppression()
+{
+	FSequenceInstance& Instance = InstanceRegistry->MutateInstance(RootInstanceHandle);
+	Instance.VolatilityManager = MoveTemp(PreviousVolatilityManager);
+	Instance.ConditionalRecompile(InstanceRegistry->GetLinker());
+}
+
 } // namespace MovieScene
 } // namespace UE
diff --git a/Engine/Source/Runtime/MovieScene/Private/EntitySystem/MovieSceneSequenceInstance.cpp b/Engine/Source/Runtime/MovieScene/Private/EntitySystem/MovieSceneSequenceInstance.cpp
--- a/Engine/Source/Runtime/MovieScene/Private/EntitySystem/MovieSceneSequenceInstance.cpp
+++ b/Engine/Source/Runtime/MovieScene/Private/EntitySystem/MovieSceneSequenceInstance.cpp
@@ -179,10 +179,6 @@ void FSequenceInstance::InvalidateCachedData(UMovieSceneEntitySystemLinker* Link
 		if (!VolatilityManager)
 		{
 			VolatilityManager = FCompiledDataVolatilityManager::Construct(*Player, RootCompiledDataID, CompiledDataManager);
-			if (VolatilityManager)
-			{
-				VolatilityManager->ConditionalRecompile(*Player, RootCompiledDataID, CompiledDataManager);
-			}
 		}
 
 		ISequenceUpdater::FactoryInstance(SequenceUpdater, CompiledDataManager, RootCompiledDataID);
diff --git a/Engine/Source/Runtime/MovieScene/Public/EntitySystem/MovieSceneInstanceRegistry.h b/Engine/Source/Runtime/MovieScene/Public/EntitySystem/MovieSceneInstanceRegistry.h
--- a/Engine/Source/Runtime/MovieScene/Public/EntitySystem/MovieSceneInstanceRegistry.h
+++ b/Engine/Source/Runtime/MovieScene/Public/EntitySystem/MovieSceneInstanceRegistry.h
@@ -122,5 +122,20 @@ private:
 };
 
 
+/**
+ * Defines a scope during which a given sequence hierarchy will not be recompiled when changed
+ * even if it was set to volatile.
+ */
+struct MOVIESCENE_API FScopedVolatilityManagerSuppression
+{
+	FScopedVolatilityManagerSuppression(FInstanceRegistry* InInstanceRegistry, FRootInstanceHandle InRootInstanceHandle);
+	~FScopedVolatilityManagerSuppression();
+
+private:
+	FInstanceRegistry* InstanceRegistry;
+	FRootInstanceHandle RootInstanceHandle;
+	TUniquePtr<FCompiledDataVolatilityManager> PreviousVolatilityManager;
+};
+
 } // namespace MovieScene
 } // namespace UE
diff --git a/Engine/Source/Runtime/MovieScene/Public/EntitySystem/MovieSceneSequenceInstance.h b/Engine/Source/Runtime/MovieScene/Public/EntitySystem/MovieSceneSequenceInstance.h
--- a/Engine/Source/Runtime/MovieScene/Public/EntitySystem/MovieSceneSequenceInstance.h
+++ b/Engine/Source/Runtime/MovieScene/Public/EntitySystem/MovieSceneSequenceInstance.h
@@ -380,6 +380,8 @@ private:
 	bool bFinished : 1;
 	/** Flag that is set if this sequence has ever updated. */
 	bool bHasEverUpdated : 1;
+
+	friend struct FScopedVolatilityManagerSuppression;
 };
 
 
