TakeRecorder: Use IsTrackClassAllowed to determine if spawnables can be created.

diff --git a/Engine/Plugins/VirtualProduction/Takes/Source/TakeRecorderSources/Private/TakeRecorderActorSource.cpp b/Engine/Plugins/VirtualProduction/Takes/Source/TakeRecorderSources/Private/TakeRecorderActorSource.cpp
--- a/Engine/Plugins/VirtualProduction/Takes/Source/TakeRecorderSources/Private/TakeRecorderActorSource.cpp
+++ b/Engine/Plugins/VirtualProduction/Takes/Source/TakeRecorderSources/Private/TakeRecorderActorSource.cpp
@@ -42,6 +42,10 @@
 #include "MovieSceneTakeSection.h"
 #include "MovieSceneTakeSettings.h"
 
+#if WITH_EDITOR
+#include "Tracks/MovieSceneSpawnTrack.h"
+#endif
+
 #include UE_INLINE_GENERATED_CPP_BY_NAME(TakeRecorderActorSource)
 
 DEFINE_LOG_CATEGORY(ActorSerialization);
@@ -1733,6 +1737,13 @@ void UTakeRecorderActorSource::SetSourceActor(TSoftObjectPtr<AActor> InTarget)
 
 bool UTakeRecorderActorSource::GetRecordToPossessable() const
 {
+#if WITH_EDITOR
+	if (!UMovieScene::IsTrackClassAllowed(UMovieSceneSpawnTrack::StaticClass()))
+	{
+		return true;
+	}
+#endif
+
 	if (RecordType == ETakeRecorderActorRecordType::ProjectDefault)
 	{
 		if (UTakeRecorderSources* Sources = GetTypedOuter<UTakeRecorderSources>())
