Make sure we have a valid sequencer before supressing recompilation.

diff --git a/Engine/Plugins/VirtualProduction/Takes/Source/TakeRecorder/Private/Recorder/TakeRecorder.cpp b/Engine/Plugins/VirtualProduction/Takes/Source/TakeRecorder/Private/Recorder/TakeRecorder.cpp
--- a/Engine/Plugins/VirtualProduction/Takes/Source/TakeRecorder/Private/Recorder/TakeRecorder.cpp
+++ b/Engine/Plugins/VirtualProduction/Takes/Source/TakeRecorder/Private/Recorder/TakeRecorder.cpp
@@ -986,9 +986,12 @@ void UTakeRecorder::Start()
 		TakeRecorderSubsystem->TakeRecorderStarted.Broadcast();
 	}
 
-	FRootInstanceHandle RootInstanceHandle = Sequencer->GetEvaluationTemplate().GetRootInstanceHandle();
-	FInstanceRegistry* InstanceRegistry = Sequencer->GetEvaluationTemplate().GetEntitySystemLinker()->GetInstanceRegistry();
-	CompileSuppression = MakeUnique<FScopedVolatilityManagerSuppression>(InstanceRegistry, RootInstanceHandle);
+	if (Sequencer.IsValid())
+	{
+		FRootInstanceHandle RootInstanceHandle = Sequencer->GetEvaluationTemplate().GetRootInstanceHandle();
+		FInstanceRegistry* InstanceRegistry = Sequencer->GetEvaluationTemplate().GetEntitySystemLinker()->GetInstanceRegistry();
+		CompileSuppression = MakeUnique<FScopedVolatilityManagerSuppression>(InstanceRegistry, RootInstanceHandle);
+	}
 }
 
 void UTakeRecorder::Stop()
