Fix an issue where caching Niagara particle in a subsequence would close the root sequence and open the subsequence in isolation

diff --git a/Engine/Plugins/VirtualProduction/Takes/Source/CacheTrackRecorder/Private/Recorder/CacheTrackRecorder.cpp b/Engine/Plugins/VirtualProduction/Takes/Source/CacheTrackRecorder/Private/Recorder/CacheTrackRecorder.cpp
--- a/Engine/Plugins/VirtualProduction/Takes/Source/CacheTrackRecorder/Private/Recorder/CacheTrackRecorder.cpp
+++ b/Engine/Plugins/VirtualProduction/Takes/Source/CacheTrackRecorder/Private/Recorder/CacheTrackRecorder.cpp
@@ -306,10 +306,11 @@ void UCacheTrackRecorder::RecordCacheTrack(IMovieSceneCachedTrack* Track, TShare
 
 void UCacheTrackRecorder::RecordCacheTracks(const TArray<IMovieSceneCachedTrack*>& CacheTracks, TSharedPtr<ISequencer> Sequencer, FCacheRecorderParameters Parameters)
 {
-	ULevelSequence* LevelSequence = Sequencer ? Cast<ULevelSequence>(Sequencer->GetFocusedMovieSceneSequence()) : nullptr;	
-	if (LevelSequence && CacheTracks.Num() > 0)
+	ULevelSequence* RootSequence = Sequencer ? Cast<ULevelSequence>(Sequencer->GetRootMovieSceneSequence()) : nullptr;	
+	ULevelSequence* FocusedSequence = Sequencer ? Cast<ULevelSequence>(Sequencer->GetFocusedMovieSceneSequence()) : nullptr;	
+	if (RootSequence && FocusedSequence && CacheTracks.Num() > 0)
 	{
-		Parameters.StartFrame = LevelSequence->GetMovieScene()->GetPlaybackRange().GetLowerBoundValue();
+		Parameters.StartFrame = FocusedSequence->GetMovieScene()->GetPlaybackRange().GetLowerBoundValue();
 		Parameters.Project.bStartAtCurrentTimecode = false;
 
 		// If not resetting the playhead, store the current time as the start frame for recording. 
@@ -320,14 +321,14 @@ void UCacheTrackRecorder::RecordCacheTracks(const TArray<IMovieSceneCachedTrack*
 		}
 
 		UCacheTrackRecorder* NewRecorder = NewObject<UCacheTrackRecorder>(GetTransientPackage(), NAME_None, RF_Transient);
-		UTakeMetaData* TakeMetaData = LevelSequence->FindOrAddMetaData<UTakeMetaData>();
+		UTakeMetaData* TakeMetaData = FocusedSequence->FindOrAddMetaData<UTakeMetaData>();
 		if (TakeMetaData->GetSlate().IsEmpty())
 		{
-			TakeMetaData->SetSlate(LevelSequence->GetName());
+			TakeMetaData->SetSlate(FocusedSequence->GetName());
 		}
 
 		FText ErrorText = LOCTEXT("UnknownError", "An unknown error occurred when trying to start recording");
-		if (!NewRecorder->Initialize(LevelSequence, CacheTracks, TakeMetaData->GetSlate(), Parameters, &ErrorText))
+		if (!NewRecorder->Initialize(RootSequence, CacheTracks, TakeMetaData->GetSlate(), Parameters, &ErrorText))
 		{
 			if (ensure(!ErrorText.IsEmpty()))
 			{
@@ -371,7 +372,7 @@ bool UCacheTrackRecorder::SetActiveRecorder(UCacheTrackRecorder* NewActiveRecord
 
 // Non-static api for UCacheTrackRecorder
 
-bool UCacheTrackRecorder::Initialize(ULevelSequence* LevelSequence, const TArray<IMovieSceneCachedTrack*>& InCacheTracks, const FString& Slate, const FCacheRecorderParameters& InParameters, FText* OutError)
+bool UCacheTrackRecorder::Initialize(ULevelSequence* RootLevelSequence, const TArray<IMovieSceneCachedTrack*>& InCacheTracks, const FString& Slate, const FCacheRecorderParameters& InParameters, FText* OutError)
 {
 	FGCObjectScopeGuard GCGuard(this);
 
@@ -405,7 +406,7 @@ bool UCacheTrackRecorder::Initialize(ULevelSequence* LevelSequence, const TArray
 	}
 
 	FCacheRecorderParameters FinalParameters = InParameters;
-	WeakSequencer = TakesUtils::OpenSequencer(LevelSequence, OutError);
+	WeakSequencer = TakesUtils::OpenSequencer(RootLevelSequence, OutError);
 	if (!WeakSequencer.IsValid())
 	{
 		return false;
@@ -632,7 +633,7 @@ void UCacheTrackRecorder::InternalTick(float DeltaTime)
 		}
 	}
 
-	if (StopRecordingFrame.IsSet() && CurrentFrameTime.FrameNumber >= StopRecordingFrame.GetValue())
+	if ((StopRecordingFrame.IsSet() && CurrentFrameTime.FrameNumber >= StopRecordingFrame.GetValue()) || Sequencer->GetPlaybackStatus() == EMovieScenePlayerStatus::Stopped)
 	{
 		Stop();
 	}
diff --git a/Engine/Plugins/VirtualProduction/Takes/Source/CacheTrackRecorder/Public/Recorder/CacheTrackRecorder.h b/Engine/Plugins/VirtualProduction/Takes/Source/CacheTrackRecorder/Public/Recorder/CacheTrackRecorder.h
--- a/Engine/Plugins/VirtualProduction/Takes/Source/CacheTrackRecorder/Public/Recorder/CacheTrackRecorder.h
+++ b/Engine/Plugins/VirtualProduction/Takes/Source/CacheTrackRecorder/Public/Recorder/CacheTrackRecorder.h
@@ -166,7 +166,7 @@ public:
 	 * Initialize a new recording with the specified parameters. Fails if another recording is currently in progress.
 	 *
 	 */
-	bool Initialize(ULevelSequence* LevelSequence, const TArray<IMovieSceneCachedTrack*>& CacheTracks, const FString& Slate, const FCacheRecorderParameters& InParameters, FText* OutError = nullptr);
+	bool Initialize(ULevelSequence* RootLevelSequence, const TArray<IMovieSceneCachedTrack*>& CacheTracks, const FString& Slate, const FCacheRecorderParameters& InParameters, FText* OutError = nullptr);
 
 	/**
 	 * Called to stop the recording
