Fixed a bug where recording tracks in subsequences with offset start positions would not play back the correct timeframe.

diff --git a/Engine/Plugins/VirtualProduction/Takes/Source/CacheTrackRecorder/Private/Recorder/CacheTrackRecorder.cpp b/Engine/Plugins/VirtualProduction/Takes/Source/CacheTrackRecorder/Private/Recorder/CacheTrackRecorder.cpp
--- a/Engine/Plugins/VirtualProduction/Takes/Source/CacheTrackRecorder/Private/Recorder/CacheTrackRecorder.cpp
+++ b/Engine/Plugins/VirtualProduction/Takes/Source/CacheTrackRecorder/Private/Recorder/CacheTrackRecorder.cpp
@@ -594,7 +594,12 @@ void UCacheTrackRecorder::Tick(float DeltaTime)
 
 FQualifiedFrameTime UCacheTrackRecorder::GetRecordTime() const
 {
-	return TakesUtils::GetRecordTime(WeakSequencer.Pin(), SequenceAsset, TimecodeAtStart, Parameters.Project.bStartAtCurrentTimecode);
+	TSharedPtr<ISequencer> Sequencer = WeakSequencer.Pin();
+	if (Sequencer.IsValid())
+	{
+		return Sequencer->GetLocalTime();
+	}
+	return TakesUtils::GetRecordTime(Sequencer, SequenceAsset, TimecodeAtStart, Parameters.Project.bStartAtCurrentTimecode);
 }
 
 void UCacheTrackRecorder::InternalTick(float DeltaTime)
@@ -737,7 +742,7 @@ void UCacheTrackRecorder::PreRecord()
 		MovieScene->SetPlaybackRange(TRange<FFrameNumber>(PlaybackStartFrame, TNumericLimits<int32>::Max() - 1), false);
 		if (Sequencer.IsValid())
 		{
-			Sequencer->SetGlobalTime(PlaybackStartFrame);
+			Sequencer->SetLocalTimeDirectly(PlaybackStartFrame);
 			Sequencer->SetPlaybackStatus(EMovieScenePlayerStatus::Paused);
 		}
 	}
