Disable auto-save when caching tracks from sequencer

diff --git a/Engine/Plugins/VirtualProduction/Takes/Source/CacheTrackRecorder/Private/Recorder/CacheTrackRecorder.cpp b/Engine/Plugins/VirtualProduction/Takes/Source/CacheTrackRecorder/Private/Recorder/CacheTrackRecorder.cpp
--- a/Engine/Plugins/VirtualProduction/Takes/Source/CacheTrackRecorder/Private/Recorder/CacheTrackRecorder.cpp
+++ b/Engine/Plugins/VirtualProduction/Takes/Source/CacheTrackRecorder/Private/Recorder/CacheTrackRecorder.cpp
@@ -698,6 +698,10 @@ void UCacheTrackRecorder::PreRecord()
 		}
 	}
 
+	// we set this global variable to prevent things like auto-save to trigger during a recording
+	GIsSlowTask = true;
+	OnStopCleanup.Add([]() { GIsSlowTask = false; });
+
 	if (Parameters.Project.bCacheTrackRecorderControlsClockTime)
 	{
 		ModifyEditorTickState();
