Fixed a bug where clicking during a cache recording could cause dropped frames from Slate throttling

diff --git a/Engine/Plugins/VirtualProduction/Takes/Source/CacheTrackRecorder/Private/Recorder/CacheTrackRecorder.cpp b/Engine/Plugins/VirtualProduction/Takes/Source/CacheTrackRecorder/Private/Recorder/CacheTrackRecorder.cpp
--- a/Engine/Plugins/VirtualProduction/Takes/Source/CacheTrackRecorder/Private/Recorder/CacheTrackRecorder.cpp
+++ b/Engine/Plugins/VirtualProduction/Takes/Source/CacheTrackRecorder/Private/Recorder/CacheTrackRecorder.cpp
@@ -36,6 +36,7 @@
 #include "Widgets/Images/SThrobber.h"
 #include "Widgets/Input/SButton.h"
 #include "Widgets/Layout/SBox.h"
+#include "Application/ThrottleManager.h"
 
 // LevelEditor includes
 #include "IAssetViewport.h"
@@ -726,6 +727,10 @@ void UCacheTrackRecorder::PreRecord()
 	GIsSlowTask = true;
 	OnStopCleanup.Add([]() { GIsSlowTask = false; });
 
+	// prevent dropped frames from slate throttling
+	FSlateThrottleManager::Get().DisableThrottle(true);
+	OnStopCleanup.Add([]() { FSlateThrottleManager::Get().DisableThrottle(false); });
+
 	if (Parameters.Project.bCacheTrackRecorderControlsClockTime)
 	{
 		ModifyEditorTickState();
