Reset sequencer speed when recording a cache and it's in reverse mode

diff --git a/Engine/Plugins/VirtualProduction/Takes/Source/CacheTrackRecorder/Private/Recorder/CacheTrackRecorder.cpp b/Engine/Plugins/VirtualProduction/Takes/Source/CacheTrackRecorder/Private/Recorder/CacheTrackRecorder.cpp
--- a/Engine/Plugins/VirtualProduction/Takes/Source/CacheTrackRecorder/Private/Recorder/CacheTrackRecorder.cpp
+++ b/Engine/Plugins/VirtualProduction/Takes/Source/CacheTrackRecorder/Private/Recorder/CacheTrackRecorder.cpp
@@ -468,6 +468,10 @@ bool UCacheTrackRecorder::Initialize(ULevelSequence* RootLevelSequence, const TA
 		TRange<double> NewRange(ViewRangeStartSeconds - 0.5f, ViewRangeStartSeconds + (Range.GetUpperBoundValue() - Range.GetLowerBoundValue()) + 0.5f);
 		Sequencer->SetViewRange(NewRange, EViewRangeInterpolation::Immediate);
 		Sequencer->SetClampRange(TRange(Sequencer->GetViewRange()));
+		if (Sequencer->GetPlaybackSpeed() <= 0)
+		{
+			Sequencer->SetPlaybackSpeed(1);
+		}
 		ESequencerLoopMode LoopMode = Sequencer->GetSequencerSettings()->GetLoopMode();
 		if (LoopMode != SLM_NoLoop)
 		{
