Take Recorder: Don't make a unique object name based on the newly added component's name because it might clash when the next newly created component.

For example, if an actor is spawning components, Comp_0, Comp_1, Comp_2, when we notice the first spawned component (Comp_0), we will make a new component based on that name. That new component will be named Comp_1 since it's based off of Comp_0. When we notice the second spawned component (Comp_1), it will already exist in the spawnable object template since we added it from the first component.

diff --git a/Engine/Plugins/VirtualProduction/Takes/Source/TakeRecorderSources/Private/TakeRecorderActorSource.cpp b/Engine/Plugins/VirtualProduction/Takes/Source/TakeRecorderSources/Private/TakeRecorderActorSource.cpp
--- a/Engine/Plugins/VirtualProduction/Takes/Source/TakeRecorderSources/Private/TakeRecorderActorSource.cpp
+++ b/Engine/Plugins/VirtualProduction/Takes/Source/TakeRecorderSources/Private/TakeRecorderActorSource.cpp
@@ -1067,7 +1067,8 @@ bool UTakeRecorderActorSource::EnsureObjectTemplateHasComponent(UActorComponent*
 
 	// Ensure the component name is unique within the Object Template. If there's complex spawn/destroy patterns that don't always use unique names this can
 	// cause UniqueComponentName to become a different name than the object it's being copied from which will cause anything attached to this to fail attachment.
-	FName UniqueComponentName = MakeUniqueObjectName(CachedObjectTemplate.Get(), InComponent->GetClass(), InComponent->GetFName());
+	// Note, we use NAME_None as the base name as opposed to anything the actual component's name because it could conflict with subsequence spawned components.
+	FName UniqueComponentName = MakeUniqueObjectName(CachedObjectTemplate.Get(), InComponent->GetClass(), NAME_None);
 	OutComponent = Cast<UActorComponent>(StaticDuplicateObject(InComponent, CachedObjectTemplate.Get(), UniqueComponentName, RF_AllFlags & ~RF_Transient));
 
 	// Restore attachment
