Fixed a bug where recording a cache didn't work when the editor window lost focus

diff --git a/Engine/Plugins/VirtualProduction/Takes/Source/CacheTrackRecorder/Private/Recorder/CacheTrackRecorder.cpp b/Engine/Plugins/VirtualProduction/Takes/Source/CacheTrackRecorder/Private/Recorder/CacheTrackRecorder.cpp
--- a/Engine/Plugins/VirtualProduction/Takes/Source/CacheTrackRecorder/Private/Recorder/CacheTrackRecorder.cpp
+++ b/Engine/Plugins/VirtualProduction/Takes/Source/CacheTrackRecorder/Private/Recorder/CacheTrackRecorder.cpp
@@ -42,6 +42,7 @@
 #include "IAssetViewport.h"
 #include "LevelEditor.h"
 #include "SequencerSettings.h"
+#include "Editor/EditorPerformanceSettings.h"
 #include "Subsystems/AssetEditorSubsystem.h"
 
 #include UE_INLINE_GENERATED_CPP_BY_NAME(CacheTrackRecorder)
@@ -735,6 +736,20 @@ void UCacheTrackRecorder::PreRecord()
 	FSlateThrottleManager::Get().DisableThrottle(true);
 	OnStopCleanup.Add([]() { FSlateThrottleManager::Get().DisableThrottle(false); });
 
+	UEditorPerformanceSettings* EditorPerformanceSettings = GetMutableDefault<UEditorPerformanceSettings>();
+	if (EditorPerformanceSettings && EditorPerformanceSettings->bThrottleCPUWhenNotForeground)
+	{
+		EditorPerformanceSettings->bThrottleCPUWhenNotForeground = false;
+		EditorPerformanceSettings->PostEditChange();
+		OnStopCleanup.Add([]()
+		{
+			UEditorPerformanceSettings* PerformanceSettings = GetMutableDefault<UEditorPerformanceSettings>();
+			PerformanceSettings->bThrottleCPUWhenNotForeground = true;
+			PerformanceSettings->PostEditChange();
+		});
+	}
+	
+
 	if (Parameters.Project.bCacheTrackRecorderControlsClockTime)
 	{
 		ModifyEditorTickState();
