When determining the subframe value, we should consider that the recording time with respect to the timecode rate if a timecode subframe was not provided by the provider.

Tested with both skeleton and basic actors and I was able to see subframe data with a time code provider that respected interframe sample between timecodes.

diff --git a/Engine/Plugins/VirtualProduction/Takes/Source/TakeRecorderSources/Private/TakeRecorderActorSource.cpp b/Engine/Plugins/VirtualProduction/Takes/Source/TakeRecorderSources/Private/TakeRecorderActorSource.cpp
--- a/Engine/Plugins/VirtualProduction/Takes/Source/TakeRecorderSources/Private/TakeRecorderActorSource.cpp
+++ b/Engine/Plugins/VirtualProduction/Takes/Source/TakeRecorderSources/Private/TakeRecorderActorSource.cpp
@@ -715,8 +715,6 @@ void UTakeRecorderActorSource::ProcessRecordedTimes(ULevelSequence* InSequence)
 				continue;
 			}
 
-			FFrameTime FrameTime = FFrameRate::TransformTime(RecordedTimePair.Key.Time, TickResolution, DisplayRate);
-
 			FTimecode Timecode = RecordedTimePair.Value.ToTimecode();
 		
 			Hours.Add(Timecode.Hours);
@@ -725,7 +723,20 @@ void UTakeRecorderActorSource::ProcessRecordedTimes(ULevelSequence* InSequence)
 			Frames.Add(Timecode.Frames);
 
 			FMovieSceneFloatValue SubFrame;
-			SubFrame.Value = RecordedTimePair.Value.Time.GetSubFrame();
+			if (RecordedTimePair.Value.Time.GetSubFrame() > 0)
+			{
+				// If the Timecode provided gave us a subframe value then we should use that value.  Otherwise, we should compute 
+				// the most appropriate value based on the timecode rate.
+				SubFrame.Value = RecordedTimePair.Value.Time.GetSubFrame();
+			}
+			else
+			{
+				FFrameTime FrameTime = FFrameRate::TransformTime(RecordedTimePair.Key.Time, TickResolution, DisplayRate);
+				FQualifiedFrameTime FrameTimeAsTimeCodeRate(FrameTime, RecordedTimePair.Value.Rate);
+
+				SubFrame.Value = FrameTimeAsTimeCodeRate.Time.GetSubFrame();
+			}
+			
 			SubFrame.InterpMode = ERichCurveInterpMode::RCIM_Linear;
 			SubFrames.Add(SubFrame);
 
