Make sure TCRate and TCSubframe carry into animation editor and export out of FBX.  The changes to TakeRecorderActorSource.cpp ensure that the animation asset has the TCRate property. The FBX change from using a FbxFloatDT to FbxDoubleDT.  FbxFloatDT properties do not show up in the anim curve editors in Maya and MotionBuilder.

diff --git a/Engine/Plugins/VirtualProduction/Takes/Source/TakeRecorderSources/Private/TakeRecorderActorSource.cpp b/Engine/Plugins/VirtualProduction/Takes/Source/TakeRecorderSources/Private/TakeRecorderActorSource.cpp
--- a/Engine/Plugins/VirtualProduction/Takes/Source/TakeRecorderSources/Private/TakeRecorderActorSource.cpp
+++ b/Engine/Plugins/VirtualProduction/Takes/Source/TakeRecorderSources/Private/TakeRecorderActorSource.cpp
@@ -886,6 +886,7 @@ FProcessRecordedTimeParams GetTimecodeRecordingParameters()
 	FString SecondsName = GetDefault<UMovieSceneTakeSettings>()->SecondsName;
 	FString FramesName = GetDefault<UMovieSceneTakeSettings>()->FramesName;
 	FString SubFramesName = GetDefault<UMovieSceneTakeSettings>()->SubFramesName;
+	FString RateName = GetDefault<UMovieSceneTakeSettings>()->RateName;
 	FString SlateName = GetDefault<UMovieSceneTakeSettings>()->SlateName;
 
 	return FProcessRecordedTimeParams {
@@ -894,6 +895,7 @@ FProcessRecordedTimeParams GetTimecodeRecordingParameters()
 		.SecondsName = SecondsName,
 		.FramesName = FramesName,
 		.SubFramesName = SubFramesName,
+		.RateName = RateName,
 		.SlateName = SlateName
 	};
 }
diff --git a/Engine/Source/Editor/UnrealEd/Private/Fbx/FbxAnimationExport.cpp b/Engine/Source/Editor/UnrealEd/Private/Fbx/FbxAnimationExport.cpp
--- a/Engine/Source/Editor/UnrealEd/Private/Fbx/FbxAnimationExport.cpp
+++ b/Engine/Source/Editor/UnrealEd/Private/Fbx/FbxAnimationExport.cpp
@@ -137,7 +137,7 @@ void FFbxExporter::ExportAnimSequenceToFbx(const UAnimSequence* AnimSeq, const U
 				}
 				else if (AttributeType == FFloatAnimationAttribute::StaticStruct())
 				{
-					FbxProperty AnimCurveFbxProp = FbxProperty::Create(CurrentBoneNode, FbxFloatDT, TCHAR_TO_UTF8(*AttributeName.ToString()));
+					FbxProperty AnimCurveFbxProp = FbxProperty::Create(CurrentBoneNode, FbxDoubleDT, TCHAR_TO_UTF8(*AttributeName.ToString()));
 					AnimCurveFbxProp.ModifyFlag(FbxPropertyFlags::eAnimatable, true);
 					AnimCurveFbxProp.ModifyFlag(FbxPropertyFlags::eUserDefined, true);
 
