Enable support for recursive classes in asset picker config so sub-classed level sequences can be used in the take recorder.

diff --git a/Engine/Plugins/VirtualProduction/Takes/Source/TakeRecorder/Private/Widgets/STakeRecorderPanel.cpp b/Engine/Plugins/VirtualProduction/Takes/Source/TakeRecorder/Private/Widgets/STakeRecorderPanel.cpp
--- a/Engine/Plugins/VirtualProduction/Takes/Source/TakeRecorder/Private/Widgets/STakeRecorderPanel.cpp
+++ b/Engine/Plugins/VirtualProduction/Takes/Source/TakeRecorder/Private/Widgets/STakeRecorderPanel.cpp
@@ -688,6 +688,7 @@ TSharedRef<SWidget> STakeRecorderPanel::OnGeneratePresetsMenu()
 
 		AssetPickerConfig.AssetShowWarningText = LOCTEXT("NoPresets_Warning", "No Presets Found");
 		AssetPickerConfig.Filter.ClassPaths.Add(UTakePreset::StaticClass()->GetClassPathName());
+		AssetPickerConfig.Filter.bRecursiveClasses = true;
 		AssetPickerConfig.OnAssetSelected = FOnAssetSelected::CreateSP(this, &STakeRecorderPanel::OnImportPreset);
 	}
 
