Add new stat sequencer category for displaying sequencer time in the stats area and option for users to enable displaying take section time.

diff --git a/Engine/Plugins/VirtualProduction/Takes/Source/TakeMovieScene/Private/MovieSceneTakeSection.cpp b/Engine/Plugins/VirtualProduction/Takes/Source/TakeMovieScene/Private/MovieSceneTakeSection.cpp
--- a/Engine/Plugins/VirtualProduction/Takes/Source/TakeMovieScene/Private/MovieSceneTakeSection.cpp
+++ b/Engine/Plugins/VirtualProduction/Takes/Source/TakeMovieScene/Private/MovieSceneTakeSection.cpp
@@ -74,6 +74,35 @@ void UMovieSceneTakeSection::PostEditImport()
 	ReconstructChannelProxy();
 }
 
+TOptional<UMovieSceneTakeSection::FSectionData> UMovieSceneTakeSection::Evaluate(FFrameTime InTime) const
+{
+	int32 Hours = 0;
+	int32 Minutes = 0;
+	int32 Seconds = 0;
+	int32 Frames = 0;
+	float Subframe = 0;;
+
+	FSectionData SectionData;
+
+	if (!HoursCurve.Evaluate(InTime, Hours)
+		|| !MinutesCurve.Evaluate(InTime, Minutes)
+		|| !SecondsCurve.Evaluate(InTime, Seconds)
+		|| !FramesCurve.Evaluate(InTime, Frames)
+		|| !SubFramesCurve.Evaluate(InTime, Subframe)
+		|| !RateCurve.Evaluate(InTime, SectionData.Rate))
+	{
+		return {};
+	}
+
+	const FString* SlateString = Slate.Evaluate(InTime);
+	bool bIsDropFrame = FTimecode::IsDropFormatTimecodeSupported(SectionData.Rate) && FTimecode::UseDropFormatTimecodeByDefaultWhenSupported();
+
+	SectionData.Timecode = FTimecode(Hours, Minutes, Seconds, Frames, Subframe, bIsDropFrame);
+	SectionData.Slate = SlateString ? *SlateString : TEXT("");
+
+	return SectionData;
+}
+
 void UMovieSceneTakeSection::ReconstructChannelProxy()
 {
 	LLM_SCOPE_BYNAME(TEXT("Takes/MovieSceneTakeSection"))
diff --git a/Engine/Plugins/VirtualProduction/Takes/Source/TakeMovieScene/Public/MovieSceneTakeSection.h b/Engine/Plugins/VirtualProduction/Takes/Source/TakeMovieScene/Public/MovieSceneTakeSection.h
--- a/Engine/Plugins/VirtualProduction/Takes/Source/TakeMovieScene/Public/MovieSceneTakeSection.h
+++ b/Engine/Plugins/VirtualProduction/Takes/Source/TakeMovieScene/Public/MovieSceneTakeSection.h
@@ -38,6 +38,16 @@ private:
 
 public:
 
+	struct FSectionData
+	{
+		FTimecode Timecode;
+		FString Slate;
+		float Rate = 24.0f;
+	};
+
+	/** Evaluate the curve data at the specified time and form a timecode / rate and slate values. If any curves failed to evaluate then the optional will be unset. */
+	TAKEMOVIESCENE_API TOptional<FSectionData> Evaluate(FFrameTime InTime) const;
+
 	/** Hours curve data */
 	UPROPERTY()
 	FMovieSceneIntegerChannel HoursCurve;
diff --git a/Engine/Plugins/VirtualProduction/Takes/Source/TakeRecorder/Private/TakeRecorderModule.cpp b/Engine/Plugins/VirtualProduction/Takes/Source/TakeRecorder/Private/TakeRecorderModule.cpp
--- a/Engine/Plugins/VirtualProduction/Takes/Source/TakeRecorder/Private/TakeRecorderModule.cpp
+++ b/Engine/Plugins/VirtualProduction/Takes/Source/TakeRecorder/Private/TakeRecorderModule.cpp
@@ -2,6 +2,7 @@
 
 #include "TakeRecorderModule.h"
 #include "TakeRecorderSettings.h"
+#include "SequencerUtilities.h"
 #include "TakeRecorderProjectSettingsCustomization.h"
 #include "Modules/ModuleManager.h"
 #include "LevelEditor.h"
@@ -29,6 +30,7 @@
 #include "Widgets/STakeRecorderPanel.h"
 
 #include "ISequencer.h"
+#include "ISequencerModule.h"
 #include "LevelSequence.h"
 #include "LevelSequenceEditorModule.h"
 #include "SequencerSettings.h"
@@ -43,6 +45,19 @@
 #include "Misc/Attribute.h"
 #include "Textures/SlateIcon.h"
 
+#include "Misc/QualifiedFrameTime.h"
+#include "MovieSceneTakeSection.h"
+#include "MovieSceneTakeTrack.h"
+#include "ISequencer.h"
+#include "LevelSequence.h"
+#include "MovieSceneSequence.h"
+#include "Sections/MovieSceneSubSection.h"
+#include "MovieScene.h"
+
+#include "Algo/RemoveIf.h"
+#include "Engine/Font.h"
+#include "CanvasTypes.h"
+
 #define LOCTEXT_NAMESPACE "TakeRecorderModule"
 
 FName ITakeRecorderModule::TakeRecorderTabName = "TakeRecorder";
@@ -67,6 +82,137 @@ TArray<ITakeRecorderDropHandler*> ITakeRecorderDropHandler::GetDropHandlers()
 	return IModularFeatures::Get().GetModularFeatureImplementations<ITakeRecorderDropHandler>(ModularFeatureName);
 }
 
+namespace UE::TakeRecorder::Private
+{
+    template<typename T>
+    concept CInvocableWithMovieSceneTakeSection = std::is_invocable_v<T, UMovieSceneTakeSection*>;
+
+	/** Render the timecode data with slate and rate to the canvas. Flag invalid TC values with a red color to indicate a problem to the user. */
+	int32 RenderTimecode(FCanvas* Canvas, int32 X, int32 Y, const FTimecode& Timecode, const float& Rate, const FString& SequenceName)
+	{
+		UFont* Font = FPlatformProperties::SupportsWindowedMode() ? GEngine->GetSmallFont() : GEngine->GetMediumFont();
+		const int32 RowHeight = FMath::TruncToInt(Font->GetMaxCharHeight());
+
+		const bool bForceSignDisplay = false;
+		const bool bAlwaysDisplaySubframe = true;
+
+		const FString TimecodeStr = Timecode.ToString(bForceSignDisplay, bAlwaysDisplaySubframe);
+
+		const FString TakeSection = TEXT("Take Section -- ");
+		const FString FPS = TEXT("(00.00)");
+		float CharWidth, CharHeight;
+		Font->GetCharSize(TEXT(' '), CharWidth, CharHeight);
+
+		int32 TakeSectionWidth = Font->GetStringSize(*TakeSection);
+		int32 NewX = X - Font->GetStringSize(*SequenceName) - Font->GetStringSize(*FPS) - TakeSectionWidth - (int32)CharWidth;
+		int32 SectionX = NewX + TakeSectionWidth;
+		FColor Color = Rate > 6000 ? FColor::Red : FColor::Green;
+		float DisplayRate = Rate > 6000 ? 0 : Rate;
+		Canvas->DrawShadowedString(NewX, Y, *TakeSection, Font, FColor::Cyan);
+		Canvas->DrawShadowedString(SectionX, Y, *FString::Printf(TEXT("%s TC: %s (%.2f)"), *SequenceName, *TimecodeStr, DisplayRate), Font, Color);
+		Y += RowHeight;
+		return Y;
+	};
+
+	template <typename CInvocableWithMovieSceneTakeSection>
+	void IterateOverMovieSceneForSections(UMovieScene* MovieScene, CInvocableWithMovieSceneTakeSection&& SectionFunction)
+	{
+		check(MovieScene);
+
+		// Lambda to iterate over tracks.  This will either get called via a FMovieSceneBinding or directly froma UMovieScene
+		auto ForEachTrack = [&SectionFunction](const TArray<UMovieSceneTrack*>& Tracks) -> bool
+		{
+			bool bDidRenderSection = false;
+			for (UMovieSceneTrack* Track : Tracks)
+			{
+				for (UMovieSceneSection* Section : Track->GetAllSections())
+				{
+					if (UMovieSceneTakeSection* TakeSection = Cast<UMovieSceneTakeSection>(Section))
+					{
+						SectionFunction(TakeSection);
+						bDidRenderSection = true;
+					}
+					else if (UMovieSceneSubSection* SubSection = Cast<UMovieSceneSubSection>(Section))
+					{
+						IterateOverMovieSceneForSections(SubSection->GetSequence()->GetMovieScene(), SectionFunction);
+					}
+				}
+			}
+			return bDidRenderSection;
+		};
+
+		const TArray<FMovieSceneBinding>& Bindings = MovieScene->GetBindings();
+		bool bDidRenderForMovieScene = false;
+		for (const FMovieSceneBinding& Binding : Bindings)
+		{
+			bDidRenderForMovieScene = ForEachTrack(Binding.GetTracks());
+		}
+
+		// If we didn't render anything with FMovieSceneBinding try the MovieScene object.
+		if (!bDidRenderForMovieScene)
+		{
+			ForEachTrack(MovieScene->GetTracks());
+		}
+	}
+
+	int32 RenderTakeSectionsInSequencer(FCanvas* Canvas, int32 X, int32 Y, TSharedPtr<ISequencer>& InSequencer)
+	{
+		const FFrameRate RootTickRate = InSequencer->GetRootTickResolution();
+		const FFrameTime CurrentTime = InSequencer->GetGlobalTime().ConvertTo(RootTickRate);
+		auto RenderOneTakeSubSection = [CurrentTime, Canvas, X, &Y](UMovieSceneTakeSection* TakeSection) mutable
+		{
+			TOptional<UMovieSceneTakeSection::FSectionData> TakeData = TakeSection->Evaluate(CurrentTime);
+			if (TakeData)
+			{
+				Y = RenderTimecode(Canvas, X, Y, TakeData->Timecode, TakeData->Rate, TakeData->Slate);
+			}
+		};
+
+		const TArray<FMovieSceneSequenceID>& SubSequenceHierarchy = InSequencer->GetSubSequenceHierarchy();
+		if (SubSequenceHierarchy.Num()>0)
+		{
+			UMovieSceneSequence* Sequence = FSequencerUtilities::GetMovieSceneSequence(InSequencer, SubSequenceHierarchy.Last());
+			check(Sequence);
+
+			UMovieScene* MovieScene = Sequence->GetMovieScene();
+			IterateOverMovieSceneForSections(MovieScene, RenderOneTakeSubSection);
+		}
+		return Y;
+	}
+
+	static FOpenSequencerWatcher SequencerWatcher;
+
+	/** For the given sequencer, iterate over all sections and find a take section then evaluate and show timecode, rate, and slate info in the HUD. */
+	int32 RenderTakeSectionTime(UWorld* World, FViewport* Viewport, FCanvas* Canvas, int32 X, int32 Y, const FVector* ViewLocation, const FRotator* ViewRotation)
+	{
+		for (const FOpenSequencerWatcher::FOpenSequencerData& OpenSequencer : SequencerWatcher.OpenSequencers)
+		{
+			if (TSharedPtr<ISequencer> Sequencer = OpenSequencer.WeakSequencer.Pin())
+			{
+				Y = RenderTakeSectionsInSequencer(Canvas, X, Y, Sequencer);
+			}
+		}
+		return Y;
+	}
+
+    void InitStatCommands()
+    {
+		auto StartupComplete = []()
+		{
+			check(GEngine);
+			if (GIsEditor)
+			{
+                const bool bIsRHS = true;
+				GEngine->AddEngineStat(TEXT("STAT_TakeTimecode"), TEXT("STATCAT_Sequencer"),
+									   LOCTEXT("TakeTimecodeDisplay", "Displays current sequencer time value in NDF timecode format."),
+									   UEngine::FEngineStatRender::CreateStatic(&RenderTakeSectionTime), nullptr, bIsRHS);
+			}
+		};
+
+		SequencerWatcher.DoStartup(StartupComplete);
+    }
+}
+
 namespace
 {
 	static TSharedRef<SDockTab> SpawnTakesBrowserTab(const FSpawnTabArgs& SpawnTabArgs)
@@ -284,6 +430,7 @@ void FTakeRecorderModule::StartupModule()
 		{
 			FCoreDelegates::OnPostEngineInit.AddRaw(this, &FTakeRecorderModule::RegisterMenus);
 		}
+		UE::TakeRecorder::Private::InitStatCommands();
 	}
 
 	if (GEditor)
diff --git a/Engine/Source/Editor/Sequencer/Private/SequencerModule.cpp b/Engine/Source/Editor/Sequencer/Private/SequencerModule.cpp
--- a/Engine/Source/Editor/Sequencer/Private/SequencerModule.cpp
+++ b/Engine/Source/Editor/Sequencer/Private/SequencerModule.cpp
@@ -20,6 +20,8 @@
 #include "Tree/CurveEditorTreeFilter.h"
 #include "AnimatedPropertyKey.h"
 #include "MovieSceneSignedObject.h"
+#include "MovieSceneSequence.h"
+#include "Sections/MovieSceneSubSection.h"
 
 #include "MVVM/CurveEditorExtension.h"
 #include "MVVM/CurveEditorIntegrationExtension.h"
@@ -56,6 +58,10 @@
 #include "PropertyEditorModule.h"
 #include "PropertyHandle.h"
 
+#include "Algo/RemoveIf.h"
+#include "Engine/Font.h"
+#include "CanvasTypes.h"
+
 #if !IS_MONOLITHIC
 UE_SELECT_ANY UE::MovieScene::FEntityManager*& GEntityManagerForDebugging = UE::MovieScene::GEntityManagerForDebuggingVisualizers;
 #endif
@@ -182,6 +188,92 @@ static void RegisterKeyframeExtensionHandler(const FOnGenerateGlobalRowExtension
 	);
 }
 
+namespace UE::SequencerModule::Private
+{
+	int32 RenderTimecode(FCanvas* Canvas, int32 X, int32 Y, const FTimecode& Timecode, const FString& SequenceName)
+	{
+		UFont* Font = FPlatformProperties::SupportsWindowedMode() ? GEngine->GetSmallFont() : GEngine->GetMediumFont();
+		const int32 RowHeight = FMath::TruncToInt(Font->GetMaxCharHeight());
+
+		const bool bForceSignDisplay = false;
+		const bool bAlwaysDisplaySubframe = true;
+
+		FString TimecodeStr = Timecode.ToString(bForceSignDisplay, bAlwaysDisplaySubframe);
+		float CharWidth, CharHeight;
+		Font->GetCharSize(TEXT(' '), CharWidth, CharHeight);
+		int32 NewX = X - Font->GetStringSize(*SequenceName) - (int32)CharWidth*14;
+
+		Canvas->DrawShadowedString(NewX, Y, *FString::Printf(TEXT("%s TC: %s"), *SequenceName, *TimecodeStr), Font, FColor::Green);
+		Y += RowHeight;
+
+		return Y;
+	};
+
+	int32 RenderTimeForSequences(FCanvas* Canvas, int32 X, int32 Y, TSharedPtr<ISequencer>& InSequencer)
+	{
+		const FFrameRate RootDisplayRate = InSequencer->GetRootDisplayRate();
+		const FFrameRate LocalDisplayRate = InSequencer->GetFocusedDisplayRate();
+		const FQualifiedFrameTime LocalCurrentTime = InSequencer->GetLocalTime();
+		const FQualifiedFrameTime RootCurrentTime = InSequencer->GetGlobalTime();
+
+		const FTimecode LocalTimecode = FTimecode::FromFrameTime(LocalCurrentTime.ConvertTo(LocalDisplayRate), LocalDisplayRate);
+		const FTimecode RootTimecode = FTimecode::FromFrameTime(RootCurrentTime.ConvertTo(RootDisplayRate), RootDisplayRate);
+
+		const TArray<FMovieSceneSequenceID>& SubSequenceHierarchy = InSequencer->GetSubSequenceHierarchy();
+		if (SubSequenceHierarchy.Num() > 0)
+		{
+			// The first one is the root sequence.
+			UMovieSceneSequence* Sequence = FSequencerUtilities::GetMovieSceneSequence(InSequencer, SubSequenceHierarchy[0]);
+			check(Sequence);
+			FString SequenceName = Sequence->GetDisplayName().ToString();
+			Y = RenderTimecode(Canvas, X, Y, RootTimecode, SequenceName);
+
+			if (SubSequenceHierarchy.Num() > 1)
+			{
+				// The current sequence is always the first in the list.
+				Sequence = FSequencerUtilities::GetMovieSceneSequence(InSequencer, SubSequenceHierarchy.Last());
+				check(Sequence);
+				SequenceName = Sequence->GetDisplayName().ToString();
+				Y = RenderTimecode(Canvas, X, Y, LocalTimecode, SequenceName);
+			}
+		}
+
+		return Y;
+	}
+
+	static FOpenSequencerWatcher SequencerWatcher;
+
+	/** Render the sequencer time to the viewport HUD. */
+	int32 RenderStatSequencerTime(UWorld* World, FViewport* Viewport, FCanvas* Canvas, int32 X, int32 Y, const FVector* ViewLocation, const FRotator* ViewRotation)
+	{
+		for (const FOpenSequencerWatcher::FOpenSequencerData& OpenSequencer : SequencerWatcher.OpenSequencers)
+		{
+			if (TSharedPtr<ISequencer> Sequencer = OpenSequencer.WeakSequencer.Pin())
+			{
+				Y = RenderTimeForSequences(Canvas, X, Y, Sequencer);
+			}
+		}
+		return Y;
+	}
+
+	void InitStatCommands()
+	{
+		auto StartupComplete = []()
+		{
+			check(GEngine);
+			if (GIsEditor)
+			{
+				const bool bIsRHS = true;
+				GEngine->AddEngineStat(TEXT("STAT_SequencerTimecode"), TEXT("STATCAT_Sequencer"),
+									   LOCTEXT("SequencerTimeDisplay", "Displays current timecode, rate, and frame for active sequencer editor."),
+									   UEngine::FEngineStatRender::CreateStatic(&RenderStatSequencerTime),
+									   nullptr, bIsRHS);
+			}
+		};
+
+		SequencerWatcher.DoStartup(StartupComplete);
+	}
+}
 /**
  * SequencerModule implementation (private)
  */
@@ -382,6 +474,7 @@ public:
 				FCoreDelegates::OnPostEngineInit.AddStatic(&FSequencerCommands::Register);
 				FCoreDelegates::OnPostEngineInit.AddRaw(this, &FSequencerModule::RegisterMenus);
 			}
+			UE::SequencerModule::Private::InitStatCommands();
 
 			FPropertyEditorModule& EditModule = FModuleManager::Get().GetModuleChecked<FPropertyEditorModule>("PropertyEditor");
 			OnGetGlobalRowExtensionHandle = EditModule.GetGlobalRowExtensionDelegate().AddStatic(&RegisterKeyframeExtensionHandler);
diff --git a/Engine/Source/Editor/Sequencer/Private/SequencerUtilities.cpp b/Engine/Source/Editor/Sequencer/Private/SequencerUtilities.cpp
--- a/Engine/Source/Editor/Sequencer/Private/SequencerUtilities.cpp
+++ b/Engine/Source/Editor/Sequencer/Private/SequencerUtilities.cpp
@@ -4435,4 +4435,51 @@ bool FSequencerUtilities::CanConvertToCustomBinding(TSharedRef<ISequencer> Seque
 	return false;
 }
 
+UMovieSceneSequence* FSequencerUtilities::GetMovieSceneSequence(TSharedPtr<ISequencer>& InSequencer, const FMovieSceneSequenceID& SequenceID)
+{
+	if (MovieSceneSequenceID::Root != SequenceID)
+	{
+		UMovieSceneSubSection* SubSection = InSequencer->FindSubSection(SequenceID);
+		return SubSection ? SubSection->GetSequence() : nullptr;
+	}
+	return InSequencer->GetRootMovieSceneSequence();
+}
+
+void FOpenSequencerWatcher::DoStartup(TFunction<void()> StartupComplete)
+{
+	auto RegisterWatcher = [this, StartupCompleteFunc = MoveTemp(StartupComplete)]()
+	{
+		ISequencerModule& SequencerModule = FModuleManager::Get().LoadModuleChecked<ISequencerModule>("Sequencer");
+		SequencerModule.RegisterOnSequencerCreated(
+			FOnSequencerCreated::FDelegate::CreateRaw(this, &FOpenSequencerWatcher::OnSequencerCreated));
+
+		StartupCompleteFunc();
+	};
+
+	if (GEngine)
+	{
+		RegisterWatcher();
+	}
+	else
+	{
+		FCoreDelegates::OnFEngineLoopInitComplete.AddLambda(RegisterWatcher);
+	}
+}
+
+void FOpenSequencerWatcher::OnSequencerCreated(TSharedRef<ISequencer> InSequencer)
+{
+	FOpenSequencerData OpenSequencer;
+	OpenSequencer.WeakSequencer = TWeakPtr<ISequencer>(InSequencer);
+	OpenSequencer.OnCloseEventHandle = InSequencer->OnCloseEvent().AddRaw(this, &FOpenSequencerWatcher::OnSequencerClosed);
+	OpenSequencers.Add(MoveTemp(OpenSequencer));
+}
+
+void FOpenSequencerWatcher::OnSequencerClosed(TSharedRef<ISequencer> InSequencer)
+{
+	OpenSequencers.RemoveAll([SequencerObject=&InSequencer.Get()](const FOpenSequencerData& Data)
+	{
+		return Data.WeakSequencer.HasSameObject(SequencerObject);
+	});
+}
+
 #undef LOCTEXT_NAMESPACE
diff --git a/Engine/Source/Editor/Sequencer/Public/SequencerUtilities.h b/Engine/Source/Editor/Sequencer/Public/SequencerUtilities.h
--- a/Engine/Source/Editor/Sequencer/Public/SequencerUtilities.h
+++ b/Engine/Source/Editor/Sequencer/Public/SequencerUtilities.h
@@ -23,6 +23,7 @@ class ISequencer;
 class FMenuBuilder;
 struct FMovieSceneBinding;
 struct FMovieScenePossessable;
+struct FMovieSceneSequenceID;
 struct FMovieSceneSpawnable;
 struct FNotificationInfo;
 class ULevelSequence;
@@ -192,6 +193,38 @@ struct FSequencerChangeBindingInfo
 	int32 BindingIndex = -1;
 };
 
+/**
+ * Helper structure to track when an sequencer is opened and closed.  Note, it will only track from when DoStartup is invoked and will not enumerate existing open sequencers.
+ */
+struct SEQUENCER_API FOpenSequencerWatcher
+{
+	/**
+	 * Begin watching the sequencer.
+	 *
+	 * @param OnStartup - will be called when we start listening for sequencer events. This will always be after engine startup if called during module load; otherwise it will be immediately called.
+	 */
+	void DoStartup(TFunction<void ()> StartupComplete);
+
+	/** Invoked when a new sequencer is created. */
+	void OnSequencerCreated(TSharedRef<ISequencer> InSequencer);
+
+	/** Registered delegate when a sequencer is closed. */
+	void OnSequencerClosed(TSharedRef<ISequencer> InSequencer);
+
+	/** Internal structure for tracking sequencers. */
+	struct FOpenSequencerData
+	{
+		/** Weak pointer to the sequencer itself, if locally opened. */
+		TWeakPtr<ISequencer> WeakSequencer;
+
+		/** Delegate handle to the Close event for the sequencer, if locally opened. */
+		FDelegateHandle OnCloseEventHandle;
+	};
+
+	/** List of open sequencers currently known by the watcher. */
+	TArray<FOpenSequencerData> OpenSequencers;
+};
+
 struct SEQUENCER_API FSequencerUtilities
 {
 	/* Creates a button (used for +Section) that opens a ComboButton with a user-defined sub-menu content. */
@@ -305,4 +338,6 @@ struct SEQUENCER_API FSequencerUtilities
 	static void AddChangeClassMenu(FMenuBuilder& MenuBuilder, TSharedRef<ISequencer> Sequencer, const TArray<FSequencerChangeBindingInfo>& BindingsToConvert, TFunction<void()> OnBindingChanged);
 	static void HandleTemplateActorClassPicked(UClass* ChosenClass, TSharedRef<ISequencer> Sequencer, const TArray<FSequencerChangeBindingInfo>& BindingsToConvert, TFunction<void()> OnBindingChanged);
 
+	/** Get a movie scene sequence from a FMovideSceneSequenceID */
+	static UMovieSceneSequence* GetMovieSceneSequence(TSharedPtr<ISequencer>& InSequencer, const FMovieSceneSequenceID& SequenceID);
 };
diff --git a/Engine/Source/Runtime/Core/Public/Misc/Timecode.h b/Engine/Source/Runtime/Core/Public/Misc/Timecode.h
--- a/Engine/Source/Runtime/Core/Public/Misc/Timecode.h
+++ b/Engine/Source/Runtime/Core/Public/Misc/Timecode.h
@@ -304,13 +304,19 @@ public:
 		return FTimecode(InTimespan.GetTotalSeconds(), InFrameRate, UseDropFormatTimecode(InFrameRate), InbRollover);
 	}
 
+	/** Drop frame is only support for frame rate of 29.97 or 59.94. */
+	static bool IsDropFormatTimecodeSupported(const double InRate)
+	{
+		return FMath::IsNearlyEqual(InRate, 30.0/1.001)
+			|| FMath::IsNearlyEqual(InRate, 60.0/1.001);
+	}
+
 	/** Drop frame is only support for frame rate of 29.97 or 59.94. */
 	static bool IsDropFormatTimecodeSupported(const FFrameRate& InFrameRate)
 	{
 		const double InRate = InFrameRate.AsDecimal();
 
-		return FMath::IsNearlyEqual(InRate, 30.0/1.001)
-			|| FMath::IsNearlyEqual(InRate, 60.0/1.001);
+		return IsDropFormatTimecodeSupported(InRate);
 	}
 
 	/** If the frame rate support drop frame format and the app wish to use drop frame format by default. */
@@ -339,8 +345,9 @@ public:
 	 * Get the Qualified Timecode formatted in HH:MM:SS:FF or HH:MM:SS;FF depending on if this represents drop-frame timecode or not.
 	 * @param bForceSignDisplay - Forces the timecode to be prepended with a positive or negative sign.
 								  Standard behavior is to only show the sign when the value is negative.
+	 * @param bAlwaysDisplaySubframe - Forces the subframe value to be included in the string result.
 	 */
-	FString ToString(bool bForceSignDisplay = false) const
+	FString ToString(bool bForceSignDisplay = false, bool bAlwaysDisplaySubframe = false) const
 	{
 		bool bHasNegativeComponent = Hours < 0 || Minutes < 0 || Seconds < 0 || Frames < 0;
 
@@ -368,7 +375,7 @@ public:
 			Builder.Appendf(TEXT("%s%02d:%02d:%02d:%02d"), SignText, FMath::Abs(Hours), FMath::Abs(Minutes), FMath::Abs(Seconds), FMath::Abs(Frames));
 		}
 
-		if (Subframe > 0)
+		if (bAlwaysDisplaySubframe || Subframe > 0)
 		{
 			int32 ClampedSubframe = static_cast<int32>(FMath::Clamp(100*Subframe,0,99));
 			Builder.Appendf(TEXT(".%02d"), ClampedSubframe);
