Move ProcessRecordedTimes to helper functions so other sources can embed take tracks in the sources.

diff --git a/Engine/Plugins/Animation/LiveLink/Source/LiveLinkSequencer/LiveLinkSequencer.Build.cs b/Engine/Plugins/Animation/LiveLink/Source/LiveLinkSequencer/LiveLinkSequencer.Build.cs
--- a/Engine/Plugins/Animation/LiveLink/Source/LiveLinkSequencer/LiveLinkSequencer.Build.cs
+++ b/Engine/Plugins/Animation/LiveLink/Source/LiveLinkSequencer/LiveLinkSequencer.Build.cs
@@ -36,6 +36,8 @@ namespace UnrealBuildTool.Rules
 					"Slate",
 					"SlateCore",
 					"TakesCore",
+					"TakeMovieScene",
+					"TakeRecorderSources",
 					"TakeRecorder",
 					"TargetPlatform",
 					"EditorFramework",
diff --git a/Engine/Plugins/Animation/LiveLink/Source/LiveLinkSequencer/Private/SequencerRecorderSections/MovieSceneLiveLinkTrackRecorder.cpp b/Engine/Plugins/Animation/LiveLink/Source/LiveLinkSequencer/Private/SequencerRecorderSections/MovieSceneLiveLinkTrackRecorder.cpp
--- a/Engine/Plugins/Animation/LiveLink/Source/LiveLinkSequencer/Private/SequencerRecorderSections/MovieSceneLiveLinkTrackRecorder.cpp
+++ b/Engine/Plugins/Animation/LiveLink/Source/LiveLinkSequencer/Private/SequencerRecorderSections/MovieSceneLiveLinkTrackRecorder.cpp
@@ -3,15 +3,20 @@
 #include "MovieSceneLiveLinkTrackRecorder.h"
 
 #include "Channels/MovieSceneChannelTraits.h"
-#include "TakeRecorderSettings.h"
 #include "Features/IModularFeatures.h"
 #include "HAL/IConsoleManager.h"
 #include "ILiveLinkClient.h"
+#include "LevelSequence.h"
 #include "LiveLinkRole.h"
 #include "LiveLinkSequencerPrivate.h"
 #include "Misc/App.h"
-#include "MovieSceneFolder.h"
+#include "Misc/QualifiedFrameTime.h"
 #include "MovieScene/MovieSceneLiveLinkTrack.h"
+#include "MovieSceneFolder.h"
+#include "MovieSceneTakeTrack.h"
+#include "TakeRecorderSettings.h"
+#include "TakeRecorderSourceHelpers.h"
+
 #include "Trace/Trace.inl"
 
 #include UE_INLINE_GENERATED_CPP_BY_NAME(MovieSceneLiveLinkTrackRecorder)
@@ -43,6 +48,8 @@ void UMovieSceneLiveLinkTrackRecorder::CreateTrack(UMovieScene* InMovieScene, co
 	bSaveSubjectSettings = bInSaveSubjectSettings;
 	bUseSourceTimecode = bInAlwaysUseTimecode;
 	bDiscardSamplesBeforeStart = bInDiscardSamplesBeforeStart;
+	bRecordTimecode = GetDefault<UTakeRecorderProjectSettings>()->Settings.bRecordTimecode;
+
 	CreateTracks();
 }
 
@@ -69,6 +76,7 @@ void UMovieSceneLiveLinkTrackRecorder::CreateTracks()
 	MovieSceneSection.Reset();
 
 	FramesToProcess.Empty();
+	RecordedTimes.Empty();
 
 	IModularFeatures& ModularFeatures = IModularFeatures::Get();
 	ILiveLinkClient* LiveLinkClient = &ModularFeatures.GetModularFeature<ILiveLinkClient>(ILiveLinkClient::ModularFeatureName);
@@ -216,6 +224,23 @@ void UMovieSceneLiveLinkTrackRecorder::FinalizeTrackImpl()
 	}
 }
 
+void UMovieSceneLiveLinkTrackRecorder::ProcessRecordedTimes(ULevelSequence* InLevelSequence)
+{
+	if (!MovieSceneSection.IsValid() || !bRecordTimecode)
+	{
+		return;
+	}
+	check(InLevelSequence);
+
+	TOptional<TRange<FFrameNumber>> FrameRange = MovieSceneSection->GetRange();
+	UMovieSceneTakeTrack* TakeTrack = Cast<UMovieSceneTakeTrack>(MovieScene->FindTrack(UMovieSceneTakeTrack::StaticClass()));
+	if (!TakeTrack)
+	{
+		TakeTrack = Cast<UMovieSceneTakeTrack>(MovieScene->AddTrack(UMovieSceneTakeTrack::StaticClass()));
+	}
+	TakeRecorderSourceHelpers::ProcessRecordedTimes(InLevelSequence, TakeTrack, FrameRange, RecordedTimes);
+}
+
 void UMovieSceneLiveLinkTrackRecorder::RecordSampleImpl(const FQualifiedFrameTime& CurrentTime)
 {
 	FTakeRecorderParameters Parameters;
@@ -247,10 +272,10 @@ void UMovieSceneLiveLinkTrackRecorder::RecordSampleImpl(const FQualifiedFrameTim
 			for (const FLiveLinkFrameDataStruct& Frame : FramesToProcess)
 			{
 				FFrameNumber FrameNumber;
-
+				FQualifiedFrameTime LiveLinkFrameTime;
 				if (bSyncedOrForced && CurrentFrameTime.IsSet())
 				{
-					FQualifiedFrameTime LiveLinkFrameTime = Frame.GetBaseData()->MetaData.SceneTime;
+					LiveLinkFrameTime = bUseSourceTimecode ? Frame.GetBaseData()->MetaData.SceneTime : *CurrentFrameTime;
 
 					if (!Parameters.Project.bStartAtCurrentTimecode)
 					{
@@ -276,6 +301,7 @@ void UMovieSceneLiveLinkTrackRecorder::RecordSampleImpl(const FQualifiedFrameTim
 					FrameNumber = (Second * TickResolution).FloorToFrame();
 					FrameNumber += RecordStartFrame;
 
+					LiveLinkFrameTime = CurrentFrameTime ? *CurrentFrameTime : FQualifiedFrameTime(FrameNumber, CurrentFrameTime.IsSet() ? CurrentFrameTime->Rate : DisplayRate);
 					UE_LOG(LogLiveLinkSequencer, VeryVerbose, TEXT("LiveLinkFrameTime (Unsynced): %f, for subject '%s'."), Frame.GetBaseData()->WorldTime.GetOffsettedTime(), *(SubjectName.ToString()));
 				}
 
@@ -283,6 +309,9 @@ void UMovieSceneLiveLinkTrackRecorder::RecordSampleImpl(const FQualifiedFrameTim
 				if (FrameNumber >= MovieSceneSection->GetInclusiveStartFrame() || !bDiscardSamplesBeforeStart)
 				{
 					MovieSceneSection->RecordFrame(FrameNumber, Frame);
+					RecordedTimes.Add(TPair<FQualifiedFrameTime, FQualifiedFrameTime>(
+						FQualifiedFrameTime(FrameNumber, TickResolution),
+						LiveLinkFrameTime));
 				}
 				else
 				{
diff --git a/Engine/Plugins/Animation/LiveLink/Source/LiveLinkSequencer/Private/TakeRecorderSource/TakeRecorderLiveLinkSource.cpp b/Engine/Plugins/Animation/LiveLink/Source/LiveLinkSequencer/Private/TakeRecorderSource/TakeRecorderLiveLinkSource.cpp
--- a/Engine/Plugins/Animation/LiveLink/Source/LiveLinkSequencer/Private/TakeRecorderSource/TakeRecorderLiveLinkSource.cpp
+++ b/Engine/Plugins/Animation/LiveLink/Source/LiveLinkSequencer/Private/TakeRecorderSource/TakeRecorderLiveLinkSource.cpp
@@ -58,6 +58,10 @@ TArray<UTakeRecorderSource*> UTakeRecorderLiveLinkSource::PostRecording(class UL
 	if (TrackRecorder)
 	{
 		TrackRecorder->FinalizeTrack();
+		if (!bCancelled)
+		{
+			TrackRecorder->ProcessRecordedTimes(InSequence);
+		}
 	}
 	
 	TrackRecorder = nullptr;
diff --git a/Engine/Plugins/Animation/LiveLink/Source/LiveLinkSequencer/Public/MovieSceneLiveLinkTrackRecorder.h b/Engine/Plugins/Animation/LiveLink/Source/LiveLinkSequencer/Public/MovieSceneLiveLinkTrackRecorder.h
--- a/Engine/Plugins/Animation/LiveLink/Source/LiveLinkSequencer/Public/MovieSceneLiveLinkTrackRecorder.h
+++ b/Engine/Plugins/Animation/LiveLink/Source/LiveLinkSequencer/Public/MovieSceneLiveLinkTrackRecorder.h
@@ -18,6 +18,7 @@ class ULiveLinkComponent;
 class ULiveLinkSubjectProperties;
 class UMovieSceneLiveLinkSection;
 class UMovieSceneLiveLinkTrack;
+class ULevelSequence;
 
 
 UCLASS(BlueprintType)
@@ -46,6 +47,8 @@ public:
 	void AddContentsToFolder(UMovieSceneFolder* InFolder);
 	void SetReduceKeys(bool bInReduce) { bReduceKeys = bInReduce; }
 
+	/** Tell the recorder to write the take track using the recorded times. */
+	void ProcessRecordedTimes(ULevelSequence *InLevelSequence);
 private:
 
 	UMovieSceneLiveLinkTrack* DoesLiveLinkTrackExist(const FName& TrackName, const TSubclassOf<ULiveLinkRole>& InTrackRole);
@@ -96,11 +99,17 @@ private:
 	/** Whether the Subject is Virtual or not*/
 	bool bIsVirtualSubject = false;
 
+	/** Should we record timecode data. */
+	bool bRecordTimecode = false;
+
 	/** Delegates registered during recording to receive live link data as it comes in*/
 	FDelegateHandle OnStaticDataReceivedHandle;
 	FDelegateHandle OnFrameDataReceivedHandle;
 
 	TArray<FLiveLinkFrameDataStruct> FramesToProcess;
+
+	/** Pairs of recorded times with the corresponding timecode.*/
+	TArray<TPair<FQualifiedFrameTime, FQualifiedFrameTime>> RecordedTimes;
 };
 
 #if UE_ENABLE_INCLUDE_ORDER_DEPRECATED_IN_5_2
diff --git a/Engine/Plugins/VirtualProduction/Takes/Source/TakeRecorderSources/Private/TakeRecorderActorSource.cpp b/Engine/Plugins/VirtualProduction/Takes/Source/TakeRecorderSources/Private/TakeRecorderActorSource.cpp
--- a/Engine/Plugins/VirtualProduction/Takes/Source/TakeRecorderSources/Private/TakeRecorderActorSource.cpp
+++ b/Engine/Plugins/VirtualProduction/Takes/Source/TakeRecorderSources/Private/TakeRecorderActorSource.cpp
@@ -12,6 +12,7 @@
 #include "Misc/ScopedSlowTask.h"
 #include "SequenceRecorderUtils.h"
 #include "TakeRecorderSource.h"
+#include "TakeRecorderSourceHelpers.h"
 #include "TakeRecorderSources.h"
 #include "TakeRecorderSourcesUtils.h"
 #include "Recorder/TakeRecorderParameters.h"
@@ -644,16 +645,12 @@ void UTakeRecorderActorSource::StopRecording(ULevelSequence* InSequence)
 	ActorSerializer.Close();
 }
 
-void UTakeRecorderActorSource::ProcessRecordedTimes(ULevelSequence* InSequence)
+namespace UE::TakeRecorderActorSource::Private
 {
-	UMovieScene* MovieScene = InSequence->GetMovieScene();
-
-	TOptional<TRange<FFrameNumber> > FrameRange;
-	FMovieSceneBinding* Binding = MovieScene->FindBinding(CachedObjectBindingGuid);
-	if (!Binding)
-	{
-		return;
-	}
+TOptional<TRange<FFrameNumber>> GetFrameRange(UMovieScene* MovieScene, FMovieSceneBinding* Binding)
+{
+	TOptional<TRange<FFrameNumber>> FrameRange;
+	check(Binding);
 
 	// In case we need it later, get the earliest timecode source *before* we
 	// add the take section, since its timecode source will be default
@@ -677,119 +674,41 @@ void UTakeRecorderActorSource::ProcessRecordedTimes(ULevelSequence* InSequence)
 			}
 		}
 	}
+	return FrameRange;
+}
 
+UMovieSceneTakeTrack* FindOrAddTakeTrack(UMovieScene* MovieScene, const FGuid& CachedObjectBindingGuid)
+{
 	UMovieSceneTakeTrack* TakeTrack = MovieScene->FindTrack<UMovieSceneTakeTrack>(CachedObjectBindingGuid);
 	if (!TakeTrack)
 	{
-		TakeTrack = InSequence->GetMovieScene()->AddTrack<UMovieSceneTakeTrack>(CachedObjectBindingGuid);
+		TakeTrack = MovieScene->AddTrack<UMovieSceneTakeTrack>(CachedObjectBindingGuid);
 	}
-	TakeTrack->RemoveAllAnimationData();
-
-	UMovieSceneTakeSection* TakeSection = Cast<UMovieSceneTakeSection>(TakeTrack->CreateNewSection());
-	TakeTrack->AddSection(*TakeSection);
-
-	if (FrameRange.IsSet())
-	{
-		TArray<int32> Hours, Minutes, Seconds, Frames;
-		TArray<FMovieSceneFloatValue> SubFrames;
-		TArray<FFrameNumber> Times;
-
-		const TArray<TPair<FQualifiedFrameTime, FQualifiedFrameTime>>& RecordedTimes = UTakeRecorderSources::RecordedTimes;
-
-		Hours.Reserve(RecordedTimes.Num());
-		Minutes.Reserve(RecordedTimes.Num());
-		Seconds.Reserve(RecordedTimes.Num());
-		Frames.Reserve(RecordedTimes.Num());
-		SubFrames.Reserve(RecordedTimes.Num());
-		Times.Reserve(RecordedTimes.Num());
-
-		FFrameRate TickResolution = MovieScene->GetTickResolution();
-		FFrameRate DisplayRate = MovieScene->GetDisplayRate();
-
-		FFrameRate TCRate = TickResolution;
-		for (const TPair<FQualifiedFrameTime, FQualifiedFrameTime>& RecordedTimePair : RecordedTimes)
-		{
-			FFrameNumber FrameNumber = RecordedTimePair.Key.Time.FrameNumber;
-			if (!FrameRange.GetValue().Contains(FrameNumber))
-			{
-				continue;
-			}
-
-			FTimecode Timecode = RecordedTimePair.Value.ToTimecode();
-			TCRate = RecordedTimePair.Value.Rate;
-			Hours.Add(Timecode.Hours);
-			Minutes.Add(Timecode.Minutes);
-			Seconds.Add(Timecode.Seconds);
-			Frames.Add(Timecode.Frames);
-
-			FMovieSceneFloatValue SubFrame;
-			if (RecordedTimePair.Value.Time.GetSubFrame() > 0)
-			{
-				// If the Timecode provided gave us a subframe value then we should use that value.  Otherwise, we should compute 
-				// the most appropriate value based on the timecode rate.
-				SubFrame.Value = RecordedTimePair.Value.Time.GetSubFrame();
-			}
-			else
-			{
-				FFrameTime FrameTime = FFrameRate::TransformTime(RecordedTimePair.Key.Time, TickResolution, DisplayRate);
-				FQualifiedFrameTime FrameTimeAsTimeCodeRate(FrameTime, TCRate);
-
-				SubFrame.Value = FrameTimeAsTimeCodeRate.Time.GetSubFrame();
-			}
-			
-			SubFrame.InterpMode = ERichCurveInterpMode::RCIM_Linear;
-			SubFrames.Add(SubFrame);
+	return TakeTrack;
+}
 
-			Times.Add(FrameNumber);
-		}
+}
+void UTakeRecorderActorSource::ProcessRecordedTimes(ULevelSequence* InSequence)
+{
+	UMovieScene* MovieScene = InSequence->GetMovieScene();
 
-		Hours.Shrink();
-		Minutes.Shrink();
-		Seconds.Shrink();
-		Frames.Shrink();
-		SubFrames.Shrink();
-		Times.Shrink();
-
-		TakeSection->HoursCurve.Set(Times, Hours);
-		TakeSection->MinutesCurve.Set(Times, Minutes);
-		TakeSection->SecondsCurve.Set(Times, Seconds);
-		TakeSection->FramesCurve.Set(Times, Frames);
-		TakeSection->SubFramesCurve.Set(Times, SubFrames);
-		TakeSection->RateCurve.SetDefault(TCRate.AsDecimal());
-	}
-
-	// Since the take section was created post recording here in this
-	// function, it wasn't available at the start of recording to have
-	// its timecode source set with the other sections, so we set it here.
-	if (TakeSection->HoursCurve.GetNumKeys() > 0)
-	{
-		// We populated the take section's timecode curves with data, so
-		// use the first values as the timecode source.
-		const int32 Hours = TakeSection->HoursCurve.GetValues()[0];
-		const int32 Minutes = TakeSection->MinutesCurve.GetValues()[0];
-		const int32 Seconds = TakeSection->SecondsCurve.GetValues()[0];
-		const int32 Frames = TakeSection->FramesCurve.GetValues()[0];
-		const bool bIsDropFrame = false;
-		const FTimecode Timecode(Hours, Minutes, Seconds, Frames, bIsDropFrame);
-		TakeSection->TimecodeSource = FMovieSceneTimecodeSource(Timecode);
-	}
-	else
+	FMovieSceneBinding* Binding = MovieScene->FindBinding(CachedObjectBindingGuid);
+	if (!Binding)
 	{
-		// Otherwise, adopt the earliest timecode source from one of the movie
-		// scene's other sections as the timecode source for the take section.
-		// This case is unlikely.
-		TakeSection->TimecodeSource = EarliestTimecodeSource;
+		return;
 	}
 
-	if (UTakeMetaData* TakeMetaData = InSequence->FindMetaData<UTakeMetaData>())
-	{
-		TakeSection->Slate.SetDefault(FString::Printf(TEXT("%s_%d"), *TakeMetaData->GetSlate(), TakeMetaData->GetTakeNumber()));
-	}
+	// In case we need it later, get the earliest timecode source *before* we
+	// add the take section, since its timecode source will be default
+	// constructed as all zeros and might accidentally compare as earliest.
+	TOptional<TRange<FFrameNumber>> FrameRange =
+		UE::TakeRecorderActorSource::Private::GetFrameRange(MovieScene, Binding);
 
-	if (TakeSection->GetAutoSizeRange().IsSet())
-	{
-		TakeSection->SetRange(TakeSection->GetAutoSizeRange().GetValue());
-	}
+	// Create a new take track or reuse the existing one based on binding.
+	UMovieSceneTakeTrack* TakeTrack = UE::TakeRecorderActorSource::Private::FindOrAddTakeTrack(MovieScene, CachedObjectBindingGuid);
+
+	// Add the recorded times to the take track.
+	TakeRecorderSourceHelpers::ProcessRecordedTimes(InSequence, TakeTrack, FrameRange, UTakeRecorderSources::RecordedTimes);
 }
 
 TArray<UTakeRecorderSource*> UTakeRecorderActorSource::PostRecording(ULevelSequence* InSequence, class ULevelSequence* InRootSequence, const bool bCancelled)
diff --git a/Engine/Plugins/VirtualProduction/Takes/Source/TakeRecorderSources/Private/TakeRecorderSourceHelpers.cpp b/Engine/Plugins/VirtualProduction/Takes/Source/TakeRecorderSources/Private/TakeRecorderSourceHelpers.cpp
--- a/Engine/Plugins/VirtualProduction/Takes/Source/TakeRecorderSources/Private/TakeRecorderSourceHelpers.cpp
+++ b/Engine/Plugins/VirtualProduction/Takes/Source/TakeRecorderSources/Private/TakeRecorderSourceHelpers.cpp
@@ -7,6 +7,13 @@
 #include "GameFramework/Actor.h"
 #include "ScopedTransaction.h"
 #include "TakeRecorderLevelSequenceSource.h"
+
+#include "LevelSequence.h"
+#include "MovieScene.h"
+#include "MovieSceneTakeTrack.h"
+#include "MovieSceneTakeSection.h"
+#include "TakeMetaData.h"
+
 #include "Editor.h"
 
 #define LOCTEXT_NAMESPACE "TakeRecorderSourceHelpers"
@@ -98,7 +105,123 @@ void RemoveAllActorSources(UTakeRecorderSources* TakeRecorderSources)
 	}
 }
 
+void ProcessRecordedTimes(ULevelSequence* InSequence, UMovieSceneTakeTrack* TakeTrack, const TOptional<TRange<FFrameNumber>>& FrameRange, const FArrayOfRecordedTimePairs& RecordedTimes)
+{
+	check(TakeTrack);
+	UMovieScene* MovieScene = InSequence->GetMovieScene();
+
+	// In case we need it later, get the earliest timecode source *before* we
+	// add the take section, since its timecode source will be default
+	// constructed as all zeros and might accidentally compare as earliest.
+	const FMovieSceneTimecodeSource EarliestTimecodeSource = MovieScene->GetEarliestTimecodeSource();
+
+	TakeTrack->RemoveAllAnimationData();
+
+	UMovieSceneTakeSection* TakeSection = Cast<UMovieSceneTakeSection>(TakeTrack->CreateNewSection());
+	TakeTrack->AddSection(*TakeSection);
+
+	if (FrameRange.IsSet())
+	{
+		TArray<int32> Hours, Minutes, Seconds, Frames;
+		TArray<FMovieSceneFloatValue> SubFrames;
+		TArray<FFrameNumber> Times;
+
+		Hours.Reserve(RecordedTimes.Num());
+		Minutes.Reserve(RecordedTimes.Num());
+		Seconds.Reserve(RecordedTimes.Num());
+		Frames.Reserve(RecordedTimes.Num());
+		SubFrames.Reserve(RecordedTimes.Num());
+		Times.Reserve(RecordedTimes.Num());
+
+		FFrameRate TickResolution = MovieScene->GetTickResolution();
+		FFrameRate DisplayRate = MovieScene->GetDisplayRate();
+
+		FFrameRate TCRate = TickResolution;
+		for (const TPair<FQualifiedFrameTime, FQualifiedFrameTime>& RecordedTimePair : RecordedTimes)
+		{
+			FFrameNumber FrameNumber = RecordedTimePair.Key.Time.FrameNumber;
+			if (!FrameRange.GetValue().Contains(FrameNumber))
+			{
+				continue;
+			}
+
+			FTimecode Timecode = RecordedTimePair.Value.ToTimecode();
+			TCRate = RecordedTimePair.Value.Rate;
+			Hours.Add(Timecode.Hours);
+			Minutes.Add(Timecode.Minutes);
+			Seconds.Add(Timecode.Seconds);
+			Frames.Add(Timecode.Frames);
+
+			FMovieSceneFloatValue SubFrame;
+			if (RecordedTimePair.Value.Time.GetSubFrame() > 0)
+			{
+				// If the Timecode provided gave us a subframe value then we should use that value.  Otherwise, we should compute
+				// the most appropriate value based on the timecode rate.
+				SubFrame.Value = RecordedTimePair.Value.Time.GetSubFrame();
+			}
+			else
+			{
+				FFrameTime FrameTime = FFrameRate::TransformTime(RecordedTimePair.Key.Time, TickResolution, DisplayRate);
+				FQualifiedFrameTime FrameTimeAsTimeCodeRate(FrameTime, TCRate);
+
+				SubFrame.Value = FrameTimeAsTimeCodeRate.Time.GetSubFrame();
+			}
+
+			SubFrame.InterpMode = ERichCurveInterpMode::RCIM_Linear;
+			SubFrames.Add(SubFrame);
+
+			Times.Add(FrameNumber);
+		}
+
+		Hours.Shrink();
+		Minutes.Shrink();
+		Seconds.Shrink();
+		Frames.Shrink();
+		SubFrames.Shrink();
+		Times.Shrink();
+
+		TakeSection->HoursCurve.Set(Times, Hours);
+		TakeSection->MinutesCurve.Set(Times, Minutes);
+		TakeSection->SecondsCurve.Set(Times, Seconds);
+		TakeSection->FramesCurve.Set(Times, Frames);
+		TakeSection->SubFramesCurve.Set(Times, SubFrames);
+		TakeSection->RateCurve.SetDefault(TCRate.AsDecimal());
+	}
+
+	// Since the take section was created post recording here in this
+	// function, it wasn't available at the start of recording to have
+	// its timecode source set with the other sections, so we set it here.
+	if (TakeSection->HoursCurve.GetNumKeys() > 0)
+	{
+		// We populated the take section's timecode curves with data, so
+		// use the first values as the timecode source.
+		const int32 Hours = TakeSection->HoursCurve.GetValues()[0];
+		const int32 Minutes = TakeSection->MinutesCurve.GetValues()[0];
+		const int32 Seconds = TakeSection->SecondsCurve.GetValues()[0];
+		const int32 Frames = TakeSection->FramesCurve.GetValues()[0];
+		const bool bIsDropFrame = false;
+		const FTimecode Timecode(Hours, Minutes, Seconds, Frames, bIsDropFrame);
+		TakeSection->TimecodeSource = FMovieSceneTimecodeSource(Timecode);
+	}
+	else
+	{
+		// Otherwise, adopt the earliest timecode source from one of the movie
+		// scene's other sections as the timecode source for the take section.
+		// This case is unlikely.
+		TakeSection->TimecodeSource = EarliestTimecodeSource;
+	}
+
+	if (UTakeMetaData* TakeMetaData = InSequence->FindMetaData<UTakeMetaData>())
+	{
+		TakeSection->Slate.SetDefault(FString::Printf(TEXT("%s_%d"), *TakeMetaData->GetSlate(), TakeMetaData->GetTakeNumber()));
+	}
+
+	if (TakeSection->GetAutoSizeRange().IsSet())
+	{
+		TakeSection->SetRange(TakeSection->GetAutoSizeRange().GetValue());
+	}
+}
 
 }
 
-#undef LOCTEXT_NAMESPACE // "TakeRecorderSourceHelpers"
\ No newline at end of file
+#undef LOCTEXT_NAMESPACE // "TakeRecorderSourceHelpers"
diff --git a/Engine/Plugins/VirtualProduction/Takes/Source/TakeRecorderSources/Public/TakeRecorderSourceHelpers.h b/Engine/Plugins/VirtualProduction/Takes/Source/TakeRecorderSources/Public/TakeRecorderSourceHelpers.h
--- a/Engine/Plugins/VirtualProduction/Takes/Source/TakeRecorderSources/Public/TakeRecorderSourceHelpers.h
+++ b/Engine/Plugins/VirtualProduction/Takes/Source/TakeRecorderSources/Public/TakeRecorderSourceHelpers.h
@@ -8,6 +8,8 @@
 #include "GameFramework/Actor.h"
 
 class AActor;
+class UMovieSceneTakeTrack;
+
 
 namespace TakeRecorderSourceHelpers
 {
@@ -29,4 +31,11 @@ namespace TakeRecorderSourceHelpers
 	 * Removes all sources from a list of sources to record from.
 	 */
 	TAKERECORDERSOURCES_API void RemoveAllActorSources(UTakeRecorderSources* Sources);
+
+	using FArrayOfRecordedTimePairs = const TArray<TPair<FQualifiedFrameTime, FQualifiedFrameTime>>;
+
+	/**
+	 * Creates a takes track to store timecode data on a take recorder source.
+	 */
+	TAKERECORDERSOURCES_API void ProcessRecordedTimes(ULevelSequence* InSequence, UMovieSceneTakeTrack* TakeTrack, const TOptional<TRange<FFrameNumber>>& FrameRange, const FArrayOfRecordedTimePairs& RecordedTimes);
 };
