Take Recorder now generates a thumbnail of the middle frame for the saved sequence.
- Introduces ISequencer::Save, which is used by [REDACTED_PERSON] to save the recorded sequence. The dedicated FSequence::Save function captures a thumbnail. Sadly, it is not possible to set up a global delegate that always captures the sequence thumbnail when saved (e.g. by selecting sequence in content browser and ctrl+s) because we require a FSequence instance.

diff --git a/Engine/Plugins/VirtualProduction/Takes/Config/BaseTakes.ini b/Engine/Plugins/VirtualProduction/Takes/Config/BaseTakes.ini
new file mode 100644
--- /dev/null
+++ b/Engine/Plugins/VirtualProduction/Takes/Config/BaseTakes.ini
@@ -0,0 +1,4 @@
+﻿[TakeRecorderSequenceEditor SequencerSettings]
+; By default, Sequencer is configured to capture the frame the user has scrubbed to (see FSequencerThumbnailCaptureSettings::CaptureFrameLocationRule).
+; In Take Recorder, the user uses no Sequencer UI - so we'll override the defaults.
+ThumbnailCaptureSettings=(CaptureFrameLocationRule=MiddleFrame)
\ No newline at end of file
diff --git a/Engine/Plugins/VirtualProduction/Takes/Source/TakeRecorder/Private/Recorder/TakeRecorder.cpp b/Engine/Plugins/VirtualProduction/Takes/Source/TakeRecorder/Private/Recorder/TakeRecorder.cpp
--- a/Engine/Plugins/VirtualProduction/Takes/Source/TakeRecorder/Private/Recorder/TakeRecorder.cpp
+++ b/Engine/Plugins/VirtualProduction/Takes/Source/TakeRecorder/Private/Recorder/TakeRecorder.cpp
@@ -1146,7 +1146,16 @@ void UTakeRecorder::StopInternal(const bool bCancelled)
 
 			if (Parameters.User.bSaveRecordedAssets)
 			{
-				TakesUtils::SaveAsset(SequenceAsset);
+				if (Sequencer)
+				{
+					// Save using the sequencer's save logic if possible - this will generate a relevant thumbnail.
+					Sequencer->Save();
+				}
+				else
+				{
+					// Fallback to "normal" saving if sequencer is not available (seems like it should be though?). Asset won't have any thumbnail.
+					TakesUtils::SaveAsset(SequenceAsset);
+				}
 			}
 
 			// Rebuild sequencer because subsequences could have been added or bindings removed
diff --git a/Engine/Source/Editor/Sequencer/Private/Sequencer.h b/Engine/Source/Editor/Sequencer/Private/Sequencer.h
--- a/Engine/Source/Editor/Sequencer/Private/Sequencer.h
+++ b/Engine/Source/Editor/Sequencer/Private/Sequencer.h
@@ -743,6 +743,7 @@ public:
 	virtual void EnterSilentMode() override;
 	virtual void ExitSilentMode() override;
 	virtual bool IsInSilentMode() const override { return SilentModeCount != 0; }
+	virtual void Save() override { SaveCurrentMovieScene(); }
 	virtual FGuid GetHandleToObject(UObject* Object, bool bCreateHandleIfMissing = true, const FName& CreatedFolderName = NAME_None) override;
 	virtual ISequencerObjectChangeListener& GetObjectChangeListener() override;
 	virtual ISequencerPropertyKeyedStatusHandler& GetPropertyKeyedStatusHandler() override;
diff --git a/Engine/Source/Editor/Sequencer/Public/ISequencer.h b/Engine/Source/Editor/Sequencer/Public/ISequencer.h
--- a/Engine/Source/Editor/Sequencer/Public/ISequencer.h
+++ b/Engine/Source/Editor/Sequencer/Public/ISequencer.h
@@ -556,6 +556,9 @@ public:
 	 */
 	virtual bool IsInSilentMode() const = 0;
 
+	/** Saves the sequence content to the asset registry. */
+	virtual void Save() = 0;
+
 	virtual FOnActorAddedToSequencer& OnActorAddedToSequencer() = 0;
 
 	DECLARE_EVENT_TwoParams(ISequencer, FOnCameraCut, UObject*, bool)
