Take Recorder now sets the "Sequencer.EnableRelevantThumbnails" console variable just before saving a sequence and restores it afterwards.

diff --git a/Engine/Plugins/VirtualProduction/Takes/Source/TakeRecorder/Private/Misc/GuardCVar.h b/Engine/Plugins/VirtualProduction/Takes/Source/TakeRecorder/Private/Misc/GuardCVar.h
new file mode 100644
--- /dev/null
+++ b/Engine/Plugins/VirtualProduction/Takes/Source/TakeRecorder/Private/Misc/GuardCVar.h
@@ -0,0 +1,53 @@
+﻿// Copyright Epic Games, Inc. All Rights Reserved.
+
+#pragma once
+
+#include "HAL/IConsoleManager.h"
+#include "Misc/Optional.h"
+
+namespace UE::TakeRecorder
+{
+	/** Guard around saving/restoring a console variable value. */
+	template<typename TCVarType>
+	class TGuardCVar
+	{
+	public:
+
+		[[nodiscard]] TGuardCVar(const TCHAR* Name, const TCVarType& NewValue)
+			: ConsoleVariable(IConsoleManager::Get().FindConsoleVariable(Name))
+			, ValueToRestore([this, &NewValue]
+			{
+				if (ensure(ConsoleVariable))
+				{
+					TCVarType OldValue;
+					ConsoleVariable->GetValue(OldValue);
+					if (OldValue != NewValue)
+					{
+						return TOptional<TCVarType>(OldValue);
+					}
+				}
+				return TOptional<TCVarType>();
+			}())
+		{
+			if (ConsoleVariable && ValueToRestore)
+			{
+				ConsoleVariable->Set(NewValue);
+			}
+		}
+
+		~TGuardCVar()
+		{
+			if (ConsoleVariable && ValueToRestore)
+			{
+				ConsoleVariable->Set(*ValueToRestore);
+			}
+		}
+
+	private:
+
+		IConsoleVariable* const ConsoleVariable = nullptr;
+		const TOptional<TCVarType> ValueToRestore;
+	};
+}
+
+
diff --git a/Engine/Plugins/VirtualProduction/Takes/Source/TakeRecorder/Private/Recorder/TakeRecorder.cpp b/Engine/Plugins/VirtualProduction/Takes/Source/TakeRecorder/Private/Recorder/TakeRecorder.cpp
--- a/Engine/Plugins/VirtualProduction/Takes/Source/TakeRecorder/Private/Recorder/TakeRecorder.cpp
+++ b/Engine/Plugins/VirtualProduction/Takes/Source/TakeRecorder/Private/Recorder/TakeRecorder.cpp
@@ -52,6 +52,7 @@
 // LevelEditor includes
 #include "IAssetViewport.h"
 #include "LevelEditor.h"
+#include "Misc/GuardCVar.h"
 #include "Subsystems/AssetEditorSubsystem.h"
 
 #include UE_INLINE_GENERATED_CPP_BY_NAME(TakeRecorder)
@@ -1148,6 +1149,10 @@ void UTakeRecorder::StopInternal(const bool bCancelled)
 			{
 				if (Sequencer)
 				{
+					// Sequencer has a feature for scrubbing to a particular frame and capturing it as thumbnail.
+					// In 5.5, this was added but hidden behind a CVar for stability reason.
+					// In BaseTakes.ini, by default, we configure our Sequencer instance to capture the MiddleFrame.
+					UE::TakeRecorder::TGuardCVar Guard(TEXT("Sequencer.EnableRelevantThumbnails"), true);
 					// Save using the sequencer's save logic if possible - this will generate a relevant thumbnail.
 					Sequencer->Save();
 				}
