Check for level sequence before attempting recursion.

diff --git a/Engine/Plugins/VirtualProduction/Takes/Source/TakeRecorder/Private/TakeRecorderModule.cpp b/Engine/Plugins/VirtualProduction/Takes/Source/TakeRecorder/Private/TakeRecorderModule.cpp
--- a/Engine/Plugins/VirtualProduction/Takes/Source/TakeRecorder/Private/TakeRecorderModule.cpp
+++ b/Engine/Plugins/VirtualProduction/Takes/Source/TakeRecorder/Private/TakeRecorderModule.cpp
@@ -117,7 +117,10 @@ namespace UE::TakeRecorder::Private
 	template <typename CInvocableWithMovieSceneTakeSection>
 	void IterateOverMovieSceneForSections(UMovieScene* MovieScene, CInvocableWithMovieSceneTakeSection&& SectionFunction)
 	{
-		check(MovieScene);
+		if (!MovieScene)
+		{
+			return;
+		}
 
 		// Lambda to iterate over tracks.  This will either get called via a FMovieSceneBinding or directly froma UMovieScene
 		auto ForEachTrack = [&SectionFunction](const TArray<UMovieSceneTrack*>& Tracks) -> bool
@@ -134,7 +137,10 @@ namespace UE::TakeRecorder::Private
 					}
 					else if (UMovieSceneSubSection* SubSection = Cast<UMovieSceneSubSection>(Section))
 					{
-						IterateOverMovieSceneForSections(SubSection->GetSequence()->GetMovieScene(), SectionFunction);
+						if (SubSection->GetSequence())
+						{
+							IterateOverMovieSceneForSections(SubSection->GetSequence()->GetMovieScene(), SectionFunction);
+						}
 					}
 				}
 			}
