Fix crash when running BootTest in QAVirtualProduction

diff --git a/Engine/Plugins/VirtualProduction/Takes/Source/TakeRecorder/Private/ScopedSequencerPanel.cpp b/Engine/Plugins/VirtualProduction/Takes/Source/TakeRecorder/Private/ScopedSequencerPanel.cpp
--- a/Engine/Plugins/VirtualProduction/Takes/Source/TakeRecorder/Private/ScopedSequencerPanel.cpp
+++ b/Engine/Plugins/VirtualProduction/Takes/Source/TakeRecorder/Private/ScopedSequencerPanel.cpp
@@ -63,7 +63,10 @@ void FScopedSequencerPanel::Open(ULevelSequence* InSequence)
 
 void FScopedSequencerPanel::Close(ULevelSequence* InSequence)
 {
-	GEditor->GetEditorSubsystem<UAssetEditorSubsystem>()->CloseAllEditorsForAsset(InSequence);
+	if (UAssetEditorSubsystem* EditorSubsystem = GEditor->GetEditorSubsystem<UAssetEditorSubsystem>())
+	{
+		EditorSubsystem->CloseAllEditorsForAsset(InSequence);
+	}
 }
 
 ECheckBoxState FScopedSequencerPanel::GetToggleCheckState() const
