Prevent users from ever recording at the default value for FFrameRate

diff --git a/Engine/Plugins/VirtualProduction/Takes/Source/TakeRecorder/Private/Widgets/STakeRecorderCockpit.cpp b/Engine/Plugins/VirtualProduction/Takes/Source/TakeRecorder/Private/Widgets/STakeRecorderCockpit.cpp
--- a/Engine/Plugins/VirtualProduction/Takes/Source/TakeRecorder/Private/Widgets/STakeRecorderCockpit.cpp
+++ b/Engine/Plugins/VirtualProduction/Takes/Source/TakeRecorder/Private/Widgets/STakeRecorderCockpit.cpp
@@ -605,6 +605,14 @@ void STakeRecorderCockpit::UpdateRecordError()
 		return;
 	}
 
+	if (TakeMetaData && TakeMetaData->GetFrameRateFromTimecode())
+	{
+		if (TakeMetaData->GetFrameRate() == FFrameRate())
+		{
+			RecordErrorText = LOCTEXT("ErrorWidget_FrameRateHigh", "The timecode rate is too high for recording.  Ensure you have a proper timecode provider set in the engine.");
+			return;
+		}
+	}
 	const int32 MaxLength = 260;
 
 	if (PackageName.Len() > MaxLength)
diff --git a/Engine/Plugins/VirtualProduction/Takes/Source/TakesCore/Private/TakeMetaData.cpp b/Engine/Plugins/VirtualProduction/Takes/Source/TakesCore/Private/TakeMetaData.cpp
--- a/Engine/Plugins/VirtualProduction/Takes/Source/TakesCore/Private/TakeMetaData.cpp
+++ b/Engine/Plugins/VirtualProduction/Takes/Source/TakesCore/Private/TakeMetaData.cpp
@@ -27,7 +27,10 @@ UTakeMetaData::UTakeMetaData(const FObjectInitializer& ObjInit)
 	TakeNumber = 1;
 	bIsLocked = false;
 	bFrameRateFromTimecode = true;
-	FrameRate = FApp::GetTimecodeFrameRate();
+	if (FApp::GetCurrentFrameTime())
+	{
+		FrameRate = FApp::GetTimecodeFrameRate();
+	}
 }
 
 UTakeMetaData* UTakeMetaData::GetConfigInstance()
diff --git a/Engine/Plugins/VirtualProduction/Takes/Source/TakesCore/Public/TakeMetaData.h b/Engine/Plugins/VirtualProduction/Takes/Source/TakesCore/Public/TakeMetaData.h
--- a/Engine/Plugins/VirtualProduction/Takes/Source/TakesCore/Public/TakeMetaData.h
+++ b/Engine/Plugins/VirtualProduction/Takes/Source/TakesCore/Public/TakeMetaData.h
@@ -287,9 +287,9 @@ private:
 	UPROPERTY(config)
 	FFrameTime Duration;
 
-	/** The frame rate the take was recorded at */
+	/** The frame rate the take was recorded at. We default to 60fps. */
 	UPROPERTY()
-	FFrameRate FrameRate;
+	FFrameRate FrameRate = FFrameRate(60, 1);
 
 	/** A user-provided description for the take */
 	UPROPERTY(config)
