Avoid linking issue with FBufferedTransformKeys in monolithic builds.   There is an idential type in SequencerRecorder that conflicts with the type defined in TrackRecorder.

diff --git a/Engine/Plugins/VirtualProduction/Takes/Source/TakeTrackRecorders/Private/MovieScene3DTransformTrackRecorder.cpp b/Engine/Plugins/VirtualProduction/Takes/Source/TakeTrackRecorders/Private/MovieScene3DTransformTrackRecorder.cpp
--- a/Engine/Plugins/VirtualProduction/Takes/Source/TakeTrackRecorders/Private/MovieScene3DTransformTrackRecorder.cpp
+++ b/Engine/Plugins/VirtualProduction/Takes/Source/TakeTrackRecorders/Private/MovieScene3DTransformTrackRecorder.cpp
@@ -235,7 +235,7 @@ void UMovieScene3DTransformTrackRecorder::FinalizeTrackImpl()
  		FirstTransform.SetScale3D(FVector(BufferedTransforms.ScaleX[0], BufferedTransforms.ScaleY[0], BufferedTransforms.ScaleZ[0]));
  	}
  
- 	BufferedTransforms = FBufferedTransformKeys();
+ 	BufferedTransforms = UE::TrackRecorders::FBufferedTransformKeys();
  
  	SlowTask.EnterProgressFrame();
  
@@ -480,7 +480,7 @@ void UMovieScene3DTransformTrackRecorder::PostProcessAnimationData(UMovieSceneAn
 	if (AnimTrackRecorder->RootWasRemoved())
 	{
 		//Get All Animation Keys
-		FBufferedTransformKeys  AnimationKeys;
+		UE::TrackRecorders::FBufferedTransformKeys  AnimationKeys;
 
 		USkeletalMeshComponent* SkeletalMeshComponent = AnimTrackRecorder->GetSkeletalMeshComponent();
 		bool bLayerCurrentBuffered = false;
diff --git a/Engine/Plugins/VirtualProduction/Takes/Source/TakeTrackRecorders/Public/TrackRecorders/MovieScene3DTransformTrackRecorder.h b/Engine/Plugins/VirtualProduction/Takes/Source/TakeTrackRecorders/Public/TrackRecorders/MovieScene3DTransformTrackRecorder.h
--- a/Engine/Plugins/VirtualProduction/Takes/Source/TakeTrackRecorders/Public/TrackRecorders/MovieScene3DTransformTrackRecorder.h
+++ b/Engine/Plugins/VirtualProduction/Takes/Source/TakeTrackRecorders/Public/TrackRecorders/MovieScene3DTransformTrackRecorder.h
@@ -40,6 +40,8 @@ public:
 };
 
 /** Structure used to buffer up transform keys. Keys are inserted into tracks in FinalizeTrack() */
+namespace UE::TrackRecorders
+{
 struct FBufferedTransformKeys
 {
 
@@ -273,6 +275,7 @@ struct FBufferedTransformKeys
 		}
 	}
 };
+}; // namespace UE::TrackRecorders
 
 UCLASS(BlueprintType)
 class TAKETRACKRECORDERS_API UMovieScene3DTransformTrackRecorder : public UMovieSceneTrackRecorder
@@ -312,7 +315,7 @@ private:
 	TWeakObjectPtr<UMovieScene3DTransformSection> MovieSceneSection;
 
 	/** Buffer of transform keys. Keys are inserted into tracks in FinalizeTrack() */
-	FBufferedTransformKeys BufferedTransforms;
+	UE::TrackRecorders::FBufferedTransformKeys BufferedTransforms;
 
 	/** The default transform this recording starts with */
 	TOptional<FTransform> DefaultTransform;
