Fix localization issue with Add Take Recorder Sources menu.

diff --git a/Engine/Plugins/VirtualProduction/Takes/Source/TakeRecorder/Private/Widgets/SLevelSequenceTakeEditor.cpp b/Engine/Plugins/VirtualProduction/Takes/Source/TakeRecorder/Private/Widgets/SLevelSequenceTakeEditor.cpp
--- a/Engine/Plugins/VirtualProduction/Takes/Source/TakeRecorder/Private/Widgets/SLevelSequenceTakeEditor.cpp
+++ b/Engine/Plugins/VirtualProduction/Takes/Source/TakeRecorder/Private/Widgets/SLevelSequenceTakeEditor.cpp
@@ -240,7 +240,7 @@ TSharedRef<SWidget> SLevelSequenceTakeEditor::OnGenerateSourcesMenu()
 			TSubclassOf<UTakeRecorderSource> SubclassOf = Class;
 
 			MenuBuilder.AddMenuEntry(
-				FText::FromString(Class->GetMetaData(TEXT("TakeRecorderDisplayName"))),
+				Class->GetMetaDataText(TEXT("TakeRecorderDisplayName")),
 				Class->GetToolTipText(true),
 				FSlateIconFinder::FindIconForClass(Class),
 				FUIAction(
@@ -992,4 +992,4 @@ void SLevelSequenceTakeEditor::UpdateDetails()
 }
 
 
-#undef LOCTEXT_NAMESPACE
\ No newline at end of file
+#undef LOCTEXT_NAMESPACE
