Fix localization issue of Take Recorder sources menu

diff --git a/Engine/Plugins/Animation/LiveLink/Source/LiveLinkSequencer/Private/TakeRecorderSource/TakeRecorderLiveLinkSource.cpp b/Engine/Plugins/Animation/LiveLink/Source/LiveLinkSequencer/Private/TakeRecorderSource/TakeRecorderLiveLinkSource.cpp
--- a/Engine/Plugins/Animation/LiveLink/Source/LiveLinkSequencer/Private/TakeRecorderSource/TakeRecorderLiveLinkSource.cpp
+++ b/Engine/Plugins/Animation/LiveLink/Source/LiveLinkSequencer/Private/TakeRecorderSource/TakeRecorderLiveLinkSource.cpp
@@ -73,6 +73,11 @@ FText UTakeRecorderLiveLinkSource::GetDisplayTextImpl() const
 	return FText::FromName(SubjectName);
 }
 
+FText UTakeRecorderLiveLinkSource::GetAddSourceDisplayTextImpl() const
+{
+    return NSLOCTEXT("UTakeRecorderLiveLinkSource", "TakeRecorderDisplayName", "Live Link");
+}
+
 void UTakeRecorderLiveLinkSource::AddContentsToFolder(UMovieSceneFolder* InFolder)
 {
 	TrackRecorder->AddContentsToFolder(InFolder);
diff --git a/Engine/Plugins/Animation/LiveLink/Source/LiveLinkSequencer/Private/TakeRecorderSource/TakeRecorderLiveLinkSource.h b/Engine/Plugins/Animation/LiveLink/Source/LiveLinkSequencer/Private/TakeRecorderSource/TakeRecorderLiveLinkSource.h
--- a/Engine/Plugins/Animation/LiveLink/Source/LiveLinkSequencer/Private/TakeRecorderSource/TakeRecorderLiveLinkSource.h
+++ b/Engine/Plugins/Animation/LiveLink/Source/LiveLinkSequencer/Private/TakeRecorderSource/TakeRecorderLiveLinkSource.h
@@ -44,7 +44,7 @@ public:
 };
 
 /** A recording source that records LiveLink */
-UCLASS(Category="Live Link", config=EditorSettings, meta = (TakeRecorderDisplayName = "Live Link"))
+UCLASS(Category="Live Link", config=EditorSettings)
 class UTakeRecorderLiveLinkSource : public UTakeRecorderSource
 {
 
@@ -97,6 +97,7 @@ private:
 	virtual bool SupportsSubscenes() const override { return true; }
 	virtual FString GetSubsceneTrackName(ULevelSequence* InSequence) const override;
 	virtual FString GetSubsceneAssetName(ULevelSequence* InSequence) const override;
+	virtual FText GetAddSourceDisplayTextImpl() const override;
 	// ~UTakeRecorderSource
 
 };
diff --git a/Engine/Plugins/Experimental/ChaosCaching/Source/ChaosCachingEditor/Private/Sequencer/TakeRecorderChaosCacheSource.cpp b/Engine/Plugins/Experimental/ChaosCaching/Source/ChaosCachingEditor/Private/Sequencer/TakeRecorderChaosCacheSource.cpp
--- a/Engine/Plugins/Experimental/ChaosCaching/Source/ChaosCachingEditor/Private/Sequencer/TakeRecorderChaosCacheSource.cpp
+++ b/Engine/Plugins/Experimental/ChaosCaching/Source/ChaosCachingEditor/Private/Sequencer/TakeRecorderChaosCacheSource.cpp
@@ -93,6 +93,11 @@ FText UTakeRecorderChaosCacheSource::GetDisplayTextImpl() const
 	return LOCTEXT("Display Text", "Chaos Cache");
 }
 
+FText UTakeRecorderChaosCacheSource::GetAddSourceDisplayTextImpl() const
+{
+    return LOCTEXT("TakeRecorderDisplayName", "Chaos Recorder");
+}
+
 void UTakeRecorderChaosCacheSource::PostEditChangeProperty(FPropertyChangedEvent& PropertyChangedEvent)
 {
 	Super::PostEditChangeProperty(PropertyChangedEvent);
diff --git a/Engine/Plugins/Experimental/ChaosCaching/Source/ChaosCachingEditor/Public/Sequencer/TakeRecorderChaosCacheSource.h b/Engine/Plugins/Experimental/ChaosCaching/Source/ChaosCachingEditor/Public/Sequencer/TakeRecorderChaosCacheSource.h
--- a/Engine/Plugins/Experimental/ChaosCaching/Source/ChaosCachingEditor/Public/Sequencer/TakeRecorderChaosCacheSource.h
+++ b/Engine/Plugins/Experimental/ChaosCaching/Source/ChaosCachingEditor/Public/Sequencer/TakeRecorderChaosCacheSource.h
@@ -16,7 +16,7 @@ class AChaosCacheManager;
 
 /** A recording source selector for the chaos integration into take recorder */
 
-UCLASS(Category = "Chaos", meta = (TakeRecorderDisplayName = "Chaos Recorder"))
+UCLASS(Category = "Chaos")
 class UTakeRecorderChaosCacheSource
 	: public UTakeRecorderSource
 {
@@ -48,6 +48,7 @@ private:
 	virtual TArray<UTakeRecorderSource*> PostRecording(class ULevelSequence* InSequence, ULevelSequence* InRootSequence, const bool bCancelled) override;
 	virtual void AddContentsToFolder(UMovieSceneFolder* InFolder) override;
 	virtual FText GetDisplayTextImpl() const override;
+	virtual FText GetAddSourceDisplayTextImpl() const override;	
 	// ~End UTakeRecorderSource
 	
 	/** Chaos cache track recorder used by this source */
diff --git a/Engine/Plugins/VirtualProduction/DMX/DMXEngine/Source/DMXEditor/Private/Sequencer/TakeRecorderDMXLibrarySource.cpp b/Engine/Plugins/VirtualProduction/DMX/DMXEngine/Source/DMXEditor/Private/Sequencer/TakeRecorderDMXLibrarySource.cpp
--- a/Engine/Plugins/VirtualProduction/DMX/DMXEngine/Source/DMXEditor/Private/Sequencer/TakeRecorderDMXLibrarySource.cpp
+++ b/Engine/Plugins/VirtualProduction/DMX/DMXEngine/Source/DMXEditor/Private/Sequencer/TakeRecorderDMXLibrarySource.cpp
@@ -129,6 +129,11 @@ FText UTakeRecorderDMXLibrarySource::GetDisplayTextImpl() const
 	return LOCTEXT("Display Text", "Null DMX Library");
 }
 
+FText UTakeRecorderDMXLibrarySource::GetAddSourceDisplayTextImpl() const
+{
+    return LOCTEXT("TakeRecorderDisplayName", "DMX Library");
+}
+
 void UTakeRecorderDMXLibrarySource::PostEditChangeProperty(FPropertyChangedEvent& PropertyChangedEvent)
 {
 	Super::PostEditChangeProperty(PropertyChangedEvent);
diff --git a/Engine/Plugins/VirtualProduction/DMX/DMXEngine/Source/DMXEditor/Private/Sequencer/TakeRecorderDMXLibrarySource.h b/Engine/Plugins/VirtualProduction/DMX/DMXEngine/Source/DMXEditor/Private/Sequencer/TakeRecorderDMXLibrarySource.h
--- a/Engine/Plugins/VirtualProduction/DMX/DMXEngine/Source/DMXEditor/Private/Sequencer/TakeRecorderDMXLibrarySource.h
+++ b/Engine/Plugins/VirtualProduction/DMX/DMXEngine/Source/DMXEditor/Private/Sequencer/TakeRecorderDMXLibrarySource.h
@@ -28,7 +28,7 @@ struct FAddAllPatchesButton
 };
 
 /** A recording source for DMX data related to a DMX Library */
-UCLASS(Category = "DMX", meta = (TakeRecorderDisplayName = "DMX Library"))
+UCLASS(Category = "DMX")
 class UTakeRecorderDMXLibrarySource
 	: public UTakeRecorderSource
 {
@@ -80,6 +80,7 @@ private:
 	virtual TArray<UTakeRecorderSource*> PostRecording(class ULevelSequence* InSequence, ULevelSequence* InRootSequence, const bool bCancelled) override;
 	virtual void AddContentsToFolder(UMovieSceneFolder* InFolder) override;
 	virtual FText GetDisplayTextImpl() const override;
+	virtual FText GetAddSourceDisplayTextImpl() const override;	
 	// ~End UTakeRecorderSource
 
 	// ~Begin UObject Interface
diff --git a/Engine/Plugins/VirtualProduction/Takes/Source/TakeRecorder/Private/Widgets/SLevelSequenceTakeEditor.cpp b/Engine/Plugins/VirtualProduction/Takes/Source/TakeRecorder/Private/Widgets/SLevelSequenceTakeEditor.cpp
--- a/Engine/Plugins/VirtualProduction/Takes/Source/TakeRecorder/Private/Widgets/SLevelSequenceTakeEditor.cpp
+++ b/Engine/Plugins/VirtualProduction/Takes/Source/TakeRecorder/Private/Widgets/SLevelSequenceTakeEditor.cpp
@@ -237,10 +237,12 @@ TSharedRef<SWidget> SLevelSequenceTakeEditor::OnGenerateSourcesMenu()
 
 		for (UClass* Class : SourceClasses)
 		{
+
 			TSubclassOf<UTakeRecorderSource> SubclassOf = Class;
 
+			UTakeRecorderSource* Default = Class->GetDefaultObject<UTakeRecorderSource>();
 			MenuBuilder.AddMenuEntry(
-				Class->GetMetaDataText(TEXT("TakeRecorderDisplayName")),
+				Default->GetAddSourceDisplayText(),
 				Class->GetToolTipText(true),
 				FSlateIconFinder::FindIconForClass(Class),
 				FUIAction(
diff --git a/Engine/Plugins/VirtualProduction/Takes/Source/TakeRecorderSources/Private/TakeRecorderActorSource.cpp b/Engine/Plugins/VirtualProduction/Takes/Source/TakeRecorderSources/Private/TakeRecorderActorSource.cpp
--- a/Engine/Plugins/VirtualProduction/Takes/Source/TakeRecorderSources/Private/TakeRecorderActorSource.cpp
+++ b/Engine/Plugins/VirtualProduction/Takes/Source/TakeRecorderSources/Private/TakeRecorderActorSource.cpp
@@ -1291,6 +1291,11 @@ FText UTakeRecorderActorSource::GetDisplayTextImpl() const
 	return LOCTEXT("ActorLabel", "Actor (None)");
 }
 
+FText UTakeRecorderActorSource::GetAddSourceDisplayTextImpl() const
+{
+    return LOCTEXT("TakeRecorderDisplayName", "Actor Any");
+}
+
 FText UTakeRecorderActorSource::GetCategoryTextImpl() const
 {
 	AActor* TargetActor = Target.Get();
diff --git a/Engine/Plugins/VirtualProduction/Takes/Source/TakeRecorderSources/Private/TakeRecorderActorSource.h b/Engine/Plugins/VirtualProduction/Takes/Source/TakeRecorderSources/Private/TakeRecorderActorSource.h
--- a/Engine/Plugins/VirtualProduction/Takes/Source/TakeRecorderSources/Private/TakeRecorderActorSource.h
+++ b/Engine/Plugins/VirtualProduction/Takes/Source/TakeRecorderSources/Private/TakeRecorderActorSource.h
@@ -31,7 +31,7 @@ enum class ETakeRecorderActorRecordType : uint8
 * Records the properties of the actor and the components on the actor and safely
 * handles new components being spawned at runtime and the actor being destroyed.
 */
-UCLASS(Category="Actors", meta = (TakeRecorderDisplayName="Any Actor"))
+UCLASS(Category="Actors")
 class UTakeRecorderActorSource : public UTakeRecorderSource, public IMovieSceneTrackRecorderHost
 {
 public:
@@ -296,7 +296,8 @@ private:
 	virtual FText GetCategoryTextImpl() const;
 	virtual FText GetDisplayTextImpl() const override;
 	virtual FText GetDescriptionTextImpl() const override;
-
+	virtual FText GetAddSourceDisplayTextImpl() const override;
+	
 	void ProcessRecordedTimes(ULevelSequence* InSequence);
 
 private:
diff --git a/Engine/Plugins/VirtualProduction/Takes/Source/TakeRecorderSources/Private/TakeRecorderCameraCutSource.cpp b/Engine/Plugins/VirtualProduction/Takes/Source/TakeRecorderSources/Private/TakeRecorderCameraCutSource.cpp
--- a/Engine/Plugins/VirtualProduction/Takes/Source/TakeRecorderSources/Private/TakeRecorderCameraCutSource.cpp
+++ b/Engine/Plugins/VirtualProduction/Takes/Source/TakeRecorderSources/Private/TakeRecorderCameraCutSource.cpp
@@ -171,4 +171,9 @@ FText UTakeRecorderCameraCutSource::GetDisplayTextImpl() const
 	return LOCTEXT("Label", "Camera Cuts");
 }
 
+FText UTakeRecorderCameraCutSource::GetAddSourceDisplayTextImpl() const
+{
+    return GetDisplayTextImpl();
+}
+
 #undef LOCTEXT_NAMESPACE // "UTakeRecorderCameraCutSource"
diff --git a/Engine/Plugins/VirtualProduction/Takes/Source/TakeRecorderSources/Private/TakeRecorderCameraCutSource.h b/Engine/Plugins/VirtualProduction/Takes/Source/TakeRecorderSources/Private/TakeRecorderCameraCutSource.h
--- a/Engine/Plugins/VirtualProduction/Takes/Source/TakeRecorderSources/Private/TakeRecorderCameraCutSource.h
+++ b/Engine/Plugins/VirtualProduction/Takes/Source/TakeRecorderSources/Private/TakeRecorderCameraCutSource.h
@@ -14,7 +14,7 @@ class UWorld;
 class UTakeRecorderActorSource;
 
 /** A recording source that detects camera switching and creates a camera cut track */
-UCLASS(Category="Other", meta = (TakeRecorderDisplayName = "Camera Cuts"))
+UCLASS(Category="Other")
 class UTakeRecorderCameraCutSource : public UTakeRecorderSource
 {
 public:
@@ -29,7 +29,7 @@ private:
 	virtual void TickRecording(const FQualifiedFrameTime& CurrentTime) override;
 	virtual TArray<UTakeRecorderSource*> PostRecording(class ULevelSequence* InSequence, class ULevelSequence* InRootSequence, const bool bCancelled) override;
 	virtual FText GetDisplayTextImpl() const override;
-
+	virtual FText GetAddSourceDisplayTextImpl() const override;		
 	// This source does not support subscenes
 	virtual bool SupportsSubscenes() const override { return false; }
 
diff --git a/Engine/Plugins/VirtualProduction/Takes/Source/TakeRecorderSources/Private/TakeRecorderLevelSequenceSource.cpp b/Engine/Plugins/VirtualProduction/Takes/Source/TakeRecorderSources/Private/TakeRecorderLevelSequenceSource.cpp
--- a/Engine/Plugins/VirtualProduction/Takes/Source/TakeRecorderSources/Private/TakeRecorderLevelSequenceSource.cpp
+++ b/Engine/Plugins/VirtualProduction/Takes/Source/TakeRecorderSources/Private/TakeRecorderLevelSequenceSource.cpp
@@ -124,6 +124,11 @@ FText UTakeRecorderLevelSequenceSource::GetDisplayTextImpl() const
 	return NSLOCTEXT("UTakeRecorderLevelSequenceSource", "LevelSequenceLabel", "Level Sequence (None)");
 }
 
+FText UTakeRecorderLevelSequenceSource::GetAddSourceDisplayTextImpl() const
+{
+    return NSLOCTEXT("UTakeRecorderLevelSequenceSource", "TakeRecorderDisplayName", "Level Sequence");
+}
+
 FText UTakeRecorderLevelSequenceSource::GetDescriptionTextImpl() const
 {
 	uint32 NumQueued = 0;
diff --git a/Engine/Plugins/VirtualProduction/Takes/Source/TakeRecorderSources/Private/TakeRecorderLevelSequenceSource.h b/Engine/Plugins/VirtualProduction/Takes/Source/TakeRecorderSources/Private/TakeRecorderLevelSequenceSource.h
--- a/Engine/Plugins/VirtualProduction/Takes/Source/TakeRecorderSources/Private/TakeRecorderLevelSequenceSource.h
+++ b/Engine/Plugins/VirtualProduction/Takes/Source/TakeRecorderSources/Private/TakeRecorderLevelSequenceSource.h
@@ -11,7 +11,7 @@ class ALevelSequenceActor;
 class ULevelSequence;
 
 /** Plays level sequence actors when recording starts */
-UCLASS(Category="Other", meta = (TakeRecorderDisplayName = "Level Sequence"))
+UCLASS(Category="Other")
 class UTakeRecorderLevelSequenceSource : public UTakeRecorderSource
 {
 public:
@@ -30,6 +30,7 @@ private:
 	virtual void StopRecording(class ULevelSequence* InSequence) override;
 	virtual FText GetDisplayTextImpl() const override;
 	virtual FText GetDescriptionTextImpl() const override;
+	virtual FText GetAddSourceDisplayTextImpl() const override;	
 
 	// This source does not support subscenes since it's a playback source instead of a recording
 	virtual bool SupportsSubscenes() const override { return false; }
diff --git a/Engine/Plugins/VirtualProduction/Takes/Source/TakeRecorderSources/Private/TakeRecorderLevelVisibilitySource.cpp b/Engine/Plugins/VirtualProduction/Takes/Source/TakeRecorderSources/Private/TakeRecorderLevelVisibilitySource.cpp
--- a/Engine/Plugins/VirtualProduction/Takes/Source/TakeRecorderSources/Private/TakeRecorderLevelVisibilitySource.cpp
+++ b/Engine/Plugins/VirtualProduction/Takes/Source/TakeRecorderSources/Private/TakeRecorderLevelVisibilitySource.cpp
@@ -146,12 +146,16 @@ void UTakeRecorderLevelVisibilitySource::AddContentsToFolder(UMovieSceneFolder*
 	}
 }
 
-
 FText UTakeRecorderLevelVisibilitySource::GetDisplayTextImpl() const
 {
 	return NSLOCTEXT("UTakeRecorderLevelVisibilitySource", "Label", "Level Visibility");
 }
 
+FText UTakeRecorderLevelVisibilitySource::GetAddSourceDisplayTextImpl() const
+{
+	return NSLOCTEXT("UTakeRecorderLevelVisibilitySource", "TakeRecorderDisplayName", "Level Visibility");
+}
+
 bool UTakeRecorderLevelVisibilitySource::CanAddSource(UTakeRecorderSources* InSources) const
 {
 	for (UTakeRecorderSource* Source : InSources->GetSources())
diff --git a/Engine/Plugins/VirtualProduction/Takes/Source/TakeRecorderSources/Private/TakeRecorderLevelVisibilitySource.h b/Engine/Plugins/VirtualProduction/Takes/Source/TakeRecorderSources/Private/TakeRecorderLevelVisibilitySource.h
--- a/Engine/Plugins/VirtualProduction/Takes/Source/TakeRecorderSources/Private/TakeRecorderLevelVisibilitySource.h
+++ b/Engine/Plugins/VirtualProduction/Takes/Source/TakeRecorderSources/Private/TakeRecorderLevelVisibilitySource.h
@@ -29,7 +29,7 @@ public:
 };
 
 /** A recording source that records level visibility state */
-UCLASS(Category="Other", config = EditorSettings, meta = (TakeRecorderDisplayName = "Level Visibility"))
+UCLASS(Category="Other", config = EditorSettings)
 class UTakeRecorderLevelVisibilitySource : public UTakeRecorderLevelVisibilitySourceSettings
 {
 public:
@@ -44,6 +44,7 @@ private:
 	virtual void TickRecording(const FQualifiedFrameTime& CurrentTime) override;
 	virtual void AddContentsToFolder(class UMovieSceneFolder* InFolder) override;
 	virtual FText GetDisplayTextImpl() const override;
+	virtual FText GetAddSourceDisplayTextImpl() const override;
 	virtual bool CanAddSource(UTakeRecorderSources* InSources) const override;
 
 private:
diff --git a/Engine/Plugins/VirtualProduction/Takes/Source/TakeRecorderSources/Private/TakeRecorderMicrophoneAudioSource.cpp b/Engine/Plugins/VirtualProduction/Takes/Source/TakeRecorderSources/Private/TakeRecorderMicrophoneAudioSource.cpp
--- a/Engine/Plugins/VirtualProduction/Takes/Source/TakeRecorderSources/Private/TakeRecorderMicrophoneAudioSource.cpp
+++ b/Engine/Plugins/VirtualProduction/Takes/Source/TakeRecorderSources/Private/TakeRecorderMicrophoneAudioSource.cpp
@@ -500,6 +500,11 @@ FText UTakeRecorderMicrophoneAudioSource::GetDisplayTextImpl() const
 	return FText::FromString(ReplaceStringTokens(AudioSourceName.ToString()));
 }
 
+FText UTakeRecorderMicrophoneAudioSource::GetAddSourceDisplayTextImpl() const
+{
+    return NSLOCTEXT("UTakeRecorderMicrophoneAudioSource", "TakeRecorderDisplayName", "Microphone Audio");
+}
+
 bool UTakeRecorderMicrophoneAudioSource::CanAddSource(UTakeRecorderSources* InSources) const
 {
 	int32 MicrophoneSourceCount = 0;
diff --git a/Engine/Plugins/VirtualProduction/Takes/Source/TakeRecorderSources/Private/TakeRecorderMicrophoneAudioSource.h b/Engine/Plugins/VirtualProduction/Takes/Source/TakeRecorderSources/Private/TakeRecorderMicrophoneAudioSource.h
--- a/Engine/Plugins/VirtualProduction/Takes/Source/TakeRecorderSources/Private/TakeRecorderMicrophoneAudioSource.h
+++ b/Engine/Plugins/VirtualProduction/Takes/Source/TakeRecorderSources/Private/TakeRecorderMicrophoneAudioSource.h
@@ -71,7 +71,7 @@ public:
 };
 
 /** A recording source that records microphone audio */
-UCLASS(Category="Audio", config=EditorSettings, meta = (TakeRecorderDisplayName = "Microphone Audio"))
+UCLASS(Category="Audio", config=EditorSettings)
 class TAKERECORDERSOURCES_API UTakeRecorderMicrophoneAudioSource : public UTakeRecorderMicrophoneAudioSourceSettings
 {
 public:
@@ -112,6 +112,8 @@ public:
 	virtual TArray<UTakeRecorderSource*> PostRecording(ULevelSequence* InSequence, class ULevelSequence* InRootSequence, const bool bCancelled) override;
 	virtual void FinalizeRecording() override;
 	virtual FText GetDisplayTextImpl() const override;
+	virtual FText GetAddSourceDisplayTextImpl() const override;
+
 	virtual void AddContentsToFolder(class UMovieSceneFolder* InFolder) override;
 	virtual bool CanAddSource(UTakeRecorderSources* InSources) const override;
 	// ~UTakeRecorderSource
diff --git a/Engine/Plugins/VirtualProduction/Takes/Source/TakeRecorderSources/Private/TakeRecorderNearbySpawnedActorSource.cpp b/Engine/Plugins/VirtualProduction/Takes/Source/TakeRecorderSources/Private/TakeRecorderNearbySpawnedActorSource.cpp
--- a/Engine/Plugins/VirtualProduction/Takes/Source/TakeRecorderSources/Private/TakeRecorderNearbySpawnedActorSource.cpp
+++ b/Engine/Plugins/VirtualProduction/Takes/Source/TakeRecorderSources/Private/TakeRecorderNearbySpawnedActorSource.cpp
@@ -100,6 +100,11 @@ FText UTakeRecorderNearbySpawnedActorSource::GetDisplayTextImpl() const
 	}
 }
 
+FText UTakeRecorderNearbySpawnedActorSource::GetAddSourceDisplayTextImpl() const
+{
+    return LOCTEXT("TakeRecorderDisplayName", "Nearby Spawned Actors");
+}
+
 FText UTakeRecorderNearbySpawnedActorSource::GetDescriptionTextImpl() const
 {
 	FNumericUnit<float> Unit(Proximity, EUnit::Centimeters);
diff --git a/Engine/Plugins/VirtualProduction/Takes/Source/TakeRecorderSources/Private/TakeRecorderNearbySpawnedActorSource.h b/Engine/Plugins/VirtualProduction/Takes/Source/TakeRecorderSources/Private/TakeRecorderNearbySpawnedActorSource.h
--- a/Engine/Plugins/VirtualProduction/Takes/Source/TakeRecorderSources/Private/TakeRecorderNearbySpawnedActorSource.h
+++ b/Engine/Plugins/VirtualProduction/Takes/Source/TakeRecorderSources/Private/TakeRecorderNearbySpawnedActorSource.h
@@ -12,7 +12,7 @@ class UTexture;
 class UTakeRecorderActorSource;
 
 /** A recording source that detects actors spawned close to the current camera, and captures them as spawnables */
-UCLASS(Category="Actors", meta = (TakeRecorderDisplayName = "Nearby Spawned Actors"))
+UCLASS(Category="Actors")
 class UTakeRecorderNearbySpawnedActorSource : public UTakeRecorderSource
 {
 public:
@@ -39,7 +39,7 @@ private:
 	virtual TArray<UTakeRecorderSource*> PostRecording(class ULevelSequence* InSequence, class ULevelSequence* InRootSequence, const bool bCancelled) override;
 	virtual FText GetDisplayTextImpl() const override;
 	virtual FText GetDescriptionTextImpl() const override;
-
+	virtual FText GetAddSourceDisplayTextImpl() const override;
 	// This source does not support subscenes (ie. "NearbySpawnedActors subscene"), but each of the spawned actors would be placed in subscenes if the option is enabled
 	virtual bool SupportsSubscenes() const override { return false; }
 
diff --git a/Engine/Plugins/VirtualProduction/Takes/Source/TakeRecorderSources/Private/TakeRecorderPlayerSource.cpp b/Engine/Plugins/VirtualProduction/Takes/Source/TakeRecorderSources/Private/TakeRecorderPlayerSource.cpp
--- a/Engine/Plugins/VirtualProduction/Takes/Source/TakeRecorderSources/Private/TakeRecorderPlayerSource.cpp
+++ b/Engine/Plugins/VirtualProduction/Takes/Source/TakeRecorderSources/Private/TakeRecorderPlayerSource.cpp
@@ -90,6 +90,11 @@ FText UTakeRecorderPlayerSource::GetDisplayTextImpl() const
 	return NSLOCTEXT("UTakeRecorderPlayerSource", "Label", "Player");
 }
 
+FText UTakeRecorderPlayerSource::GetAddSourceDisplayTextImpl() const
+{
+	return GetDisplayTextImpl();
+}
+
 bool UTakeRecorderPlayerSource::CanAddSource(UTakeRecorderSources* InSources) const
 {
 	for (UTakeRecorderSource* Source : InSources->GetSources())
diff --git a/Engine/Plugins/VirtualProduction/Takes/Source/TakeRecorderSources/Private/TakeRecorderPlayerSource.h b/Engine/Plugins/VirtualProduction/Takes/Source/TakeRecorderSources/Private/TakeRecorderPlayerSource.h
--- a/Engine/Plugins/VirtualProduction/Takes/Source/TakeRecorderSources/Private/TakeRecorderPlayerSource.h
+++ b/Engine/Plugins/VirtualProduction/Takes/Source/TakeRecorderSources/Private/TakeRecorderPlayerSource.h
@@ -10,7 +10,7 @@
 class UTakeRecorderActorSource;
 
 /** A recording source that records the current player */
-UCLASS(Category="Actors", meta = (TakeRecorderDisplayName = "Player"))
+UCLASS(Category="Actors")
 class UTakeRecorderPlayerSource : public UTakeRecorderSource
 {
 public:
@@ -24,6 +24,7 @@ private:
 	virtual TArray<UTakeRecorderSource*> PreRecording(ULevelSequence* InSequence, FMovieSceneSequenceID InSequenceID, ULevelSequence* InRootSequence, FManifestSerializer* InManifestSerializer) override;
 	virtual TArray<UTakeRecorderSource*> PostRecording(class ULevelSequence* InSequence, class ULevelSequence* InRootSequence, const bool bCancelled) override;
 	virtual FText GetDisplayTextImpl() const override;
+	virtual FText GetAddSourceDisplayTextImpl() const override;				
 	virtual bool CanAddSource(UTakeRecorderSources* InSources) const override;
 
 	// This source does not support subscenes (ie. "Player subscene"), but the player would be placed in subscenes if the option is enabled
diff --git a/Engine/Plugins/VirtualProduction/Takes/Source/TakeRecorderSources/Private/TakeRecorderWorldSource.cpp b/Engine/Plugins/VirtualProduction/Takes/Source/TakeRecorderSources/Private/TakeRecorderWorldSource.cpp
--- a/Engine/Plugins/VirtualProduction/Takes/Source/TakeRecorderSources/Private/TakeRecorderWorldSource.cpp
+++ b/Engine/Plugins/VirtualProduction/Takes/Source/TakeRecorderSources/Private/TakeRecorderWorldSource.cpp
@@ -117,6 +117,11 @@ FText UTakeRecorderWorldSource::GetDisplayTextImpl() const
 	return NSLOCTEXT("UTakeRecorderWorldSource", "Label", "World");
 }
 
+FText UTakeRecorderWorldSource::GetAddSourceDisplayTextImpl() const
+{
+    return NSLOCTEXT("UTakeRecorderWorldSource", "TakeRecorderDisplayName", "World Recorder");
+}
+
 bool UTakeRecorderWorldSource::CanAddSource(UTakeRecorderSources* InSources) const
 {
 	for (UTakeRecorderSource* Source : InSources->GetSources())
diff --git a/Engine/Plugins/VirtualProduction/Takes/Source/TakeRecorderSources/Private/TakeRecorderWorldSource.h b/Engine/Plugins/VirtualProduction/Takes/Source/TakeRecorderSources/Private/TakeRecorderWorldSource.h
--- a/Engine/Plugins/VirtualProduction/Takes/Source/TakeRecorderSources/Private/TakeRecorderWorldSource.h
+++ b/Engine/Plugins/VirtualProduction/Takes/Source/TakeRecorderSources/Private/TakeRecorderWorldSource.h
@@ -34,7 +34,7 @@ public:
 
 
 /** A recording source that records world state */
-UCLASS(Category="Actors", meta = (TakeRecorderDisplayName = "World"))
+UCLASS(Category="Actors")
 class UTakeRecorderWorldSource : public UTakeRecorderWorldSourceSettings
 {
 public:
@@ -49,6 +49,7 @@ private:
 	virtual TArray<UTakeRecorderSource*> PostRecording(class ULevelSequence* InSequence, class ULevelSequence* InRootSequence, const bool bCancelled) override;
 	virtual bool SupportsTakeNumber() const override { return false; }
 	virtual FText GetDisplayTextImpl() const override;
+	virtual FText GetAddSourceDisplayTextImpl() const override;		
 	virtual bool CanAddSource(UTakeRecorderSources* InSources) const override;
 
 	// This source does not support subscenes (ie. "World Settings subscene"), but the world settings actor would be placed in subscenes if the option is enabled
diff --git a/Engine/Plugins/VirtualProduction/Takes/Source/TakesCore/Public/TakeRecorderSource.h b/Engine/Plugins/VirtualProduction/Takes/Source/TakesCore/Public/TakeRecorderSource.h
--- a/Engine/Plugins/VirtualProduction/Takes/Source/TakesCore/Public/TakeRecorderSource.h
+++ b/Engine/Plugins/VirtualProduction/Takes/Source/TakesCore/Public/TakeRecorderSource.h
@@ -185,6 +185,20 @@ public:
 		return GetDisplayTextImpl();
 	}
 
+	/**
+	 * Get the text to display on the add source menu for this source
+	 */
+	FText GetAddSourceDisplayText() const
+	{
+		FText Result = GetAddSourceDisplayTextImpl();
+		if (Result.IsEmpty())
+		{
+			ensureMsgf(false, TEXT("You must provide localized source name."));
+			return GetDisplayText();
+		}
+		return Result;
+	}
+	
 	/**
 	 * Get the icon to display on the Take Recorder source list for this source
 	 */
@@ -213,4 +227,5 @@ private:
 	virtual FText GetCategoryTextImpl() const { return FText(); }
 	virtual FText GetDisplayTextImpl() const { return FText(); }
 	virtual FText GetDescriptionTextImpl() const { return FText(); }
+	virtual FText GetAddSourceDisplayTextImpl() const { return FText(); };
 };
