Add some instrumentation to TakeRecorder so that we can capture cost of various stages in TakeRecorder phases.

diff --git a/Engine/Plugins/VirtualProduction/Takes/Source/TakeRecorder/Private/Recorder/TakeRecorder.cpp b/Engine/Plugins/VirtualProduction/Takes/Source/TakeRecorder/Private/Recorder/TakeRecorder.cpp
--- a/Engine/Plugins/VirtualProduction/Takes/Source/TakeRecorder/Private/Recorder/TakeRecorder.cpp
+++ b/Engine/Plugins/VirtualProduction/Takes/Source/TakeRecorder/Private/Recorder/TakeRecorder.cpp
@@ -24,6 +24,7 @@
 #include "TakePreset.h"
 #include "TakeRecorderOverlayWidget.h"
 #include "TakeRecorderSources.h"
+#include "TakeRecorderTracing.h"
 #include "TakesUtils.h"
 #include "Tickable.h"
 #include "Stats/Stats.h"
@@ -435,6 +436,7 @@ FTakeRecorderParameters AccumulateParamsOverride(const FTakeRecorderParameters&
 
 bool UTakeRecorder::Initialize ( ULevelSequence* LevelSequenceBase, UTakeRecorderSources* Sources, UTakeMetaData* MetaData, const FTakeRecorderParameters& InParameters, FText* OutError )
 {
+	SCOPED_TAKERECORDER_TRACE(UTakeRecorder_Initialize);
 	FGCObjectScopeGuard GCGuard(this);
 
 	if (GetActiveRecorder())
@@ -757,6 +759,7 @@ FQualifiedFrameTime UTakeRecorder::GetRecordTime() const
 
 void UTakeRecorder::InternalTick(float DeltaTime)
 {
+	SCOPED_TAKERECORDER_TRACE(UTakeRecorder_InternalTick);
 	using namespace UE::MovieScene;
 
 	FScopedSignedObjectModifyDefer FlushOnTick(true);
@@ -797,6 +800,8 @@ void UTakeRecorder::InternalTick(float DeltaTime)
 
 void UTakeRecorder::PreRecord()
 {
+	SCOPED_TAKERECORDER_TRACE(UTakeRecorder_PreRecord);
+	
 	State = ETakeRecorderState::PreRecord;
 
 	UTakeRecorderSources* Sources = SequenceAsset->FindMetaData<UTakeRecorderSources>();
@@ -900,6 +905,7 @@ void UTakeRecorder::PreRecord()
 
 void UTakeRecorder::Start()
 {
+	SCOPED_TAKERECORDER_TRACE(UTakeRecorder_Start);
 	using namespace UE::MovieScene;
 
 	FTimecode Timecode = FApp::GetTimecode();
@@ -967,6 +973,8 @@ void UTakeRecorder::Start()
 
 	if (!Parameters.bDisableRecordingAndSave)
 	{
+		SCOPED_TAKERECORDER_TRACE(UTakeRecorder_StartRecordingSources);
+		
 		FQualifiedFrameTime RecordTime = GetRecordTime();
 
 		Sources->StartRecording(SequenceAsset, RecordTime, Parameters.User.bAutoSerialize ? &ManifestSerializer : nullptr);
@@ -1009,6 +1017,7 @@ void UTakeRecorder::Cancel()
 
 void UTakeRecorder::StopInternal(const bool bCancelled)
 {
+	SCOPED_TAKERECORDER_TRACE(UTakeRecorder_StopInternal);
 	static bool bStoppedRecording = false;
 
 	if (bStoppedRecording)
diff --git a/Engine/Plugins/VirtualProduction/Takes/Source/TakeRecorder/Private/TakeRecorderTracing.cpp b/Engine/Plugins/VirtualProduction/Takes/Source/TakeRecorder/Private/TakeRecorderTracing.cpp
new file mode 100644
--- /dev/null
+++ b/Engine/Plugins/VirtualProduction/Takes/Source/TakeRecorder/Private/TakeRecorderTracing.cpp
@@ -0,0 +1,5 @@
+// Copyright Epic Games, Inc. All Rights Reserved.
+
+#include "TakeRecorderTracing.h"
+
+UE_TRACE_CHANNEL_DEFINE(TakeRecorderChannel);
diff --git a/Engine/Plugins/VirtualProduction/Takes/Source/TakeRecorder/Public/TakeRecorderTracing.h b/Engine/Plugins/VirtualProduction/Takes/Source/TakeRecorder/Public/TakeRecorderTracing.h
new file mode 100644
--- /dev/null
+++ b/Engine/Plugins/VirtualProduction/Takes/Source/TakeRecorder/Public/TakeRecorderTracing.h
@@ -0,0 +1,10 @@
+// Copyright Epic Games, Inc. All Rights Reserved.
+
+#pragma once
+
+#include "ProfilingDebugging/CpuProfilerTrace.h"
+#include "Trace/Trace.h"
+
+UE_TRACE_CHANNEL_EXTERN(TakeRecorderChannel, TAKERECORDER_API);
+#define SCOPED_TAKERECORDER_TRACE(TraceName) TRACE_CPUPROFILER_EVENT_SCOPE_ON_CHANNEL(TraceName, TakeRecorderChannel)
+
diff --git a/Engine/Plugins/VirtualProduction/Takes/Source/TakeRecorder/TakeRecorder.Build.cs b/Engine/Plugins/VirtualProduction/Takes/Source/TakeRecorder/TakeRecorder.Build.cs
--- a/Engine/Plugins/VirtualProduction/Takes/Source/TakeRecorder/TakeRecorder.Build.cs
+++ b/Engine/Plugins/VirtualProduction/Takes/Source/TakeRecorder/TakeRecorder.Build.cs
@@ -46,6 +46,7 @@ public class TakeRecorder : ModuleRules
                 "TakeTrackRecorders",
                 "SerializedRecorderInterface",
                 "Sequencer",
+				"TraceLog"
             }
         );
     }
