TakeRecorder: Fix incorrect NamingTokens filename.

diff --git a/Engine/Plugins/VirtualProduction/Takes/Source/TakeRecorderNamingTokens/Private/TakeRecoderNamingTokens.cpp b/Engine/Plugins/VirtualProduction/Takes/Source/TakeRecorderNamingTokens/Private/TakeRecoderNamingTokens.cpp
deleted file mode 100644
diff --git a/Engine/Plugins/VirtualProduction/Takes/Source/TakeRecorderNamingTokens/Private/TakeRecorderNamingTokens.cpp b/Engine/Plugins/VirtualProduction/Takes/Source/TakeRecorderNamingTokens/Private/TakeRecorderNamingTokens.cpp
new file mode 100644
--- /dev/null
+++ b/Engine/Plugins/VirtualProduction/Takes/Source/TakeRecorderNamingTokens/Private/TakeRecorderNamingTokens.cpp
@@ -0,0 +1,154 @@
+// Copyright Epic Games, Inc. All Rights Reserved.
+
+#include "TakeRecorderNamingTokens.h"
+
+#include "Editor.h"
+#include "ITakeRecorderNamingTokensModule.h"
+#include "Misc/PackageName.h"
+#include "MovieSceneToolsProjectSettings.h"
+#include "TakeMetaData.h"
+#include "TakeRecorderSettings.h"
+#include "TakesCoreBlueprintLibrary.h"
+
+#define LOCTEXT_NAMESPACE "TakeRecorderNamingTokens"
+
+UTakeRecorderNamingTokens::UTakeRecorderNamingTokens()
+{
+	Namespace = ITakeRecorderNamingTokensModule::GetTakeRecorderNamespace();
+}
+
+UTakeRecorderNamingTokens::~UTakeRecorderNamingTokens()
+{
+}
+
+void UTakeRecorderNamingTokens::OnCreateDefaultTokens(TArray<FNamingTokenData>& Tokens)
+{
+	Super::OnCreateDefaultTokens(Tokens);
+
+	Tokens.Add({
+	    TEXT("day"),
+	    LOCTEXT("TokenDay", "Day"),
+	    FNamingTokenData::FTokenProcessorDelegateNative::CreateLambda([this]{
+	        return FText::FromString(FString::Printf(TEXT("%02i"), GetCurrentDateTime().GetDay()));
+	    })
+	});
+
+	Tokens.Add({
+	    TEXT("month"),
+	    LOCTEXT("TokenMonth", "Month"),
+	    FNamingTokenData::FTokenProcessorDelegateNative::CreateLambda([this]{
+	        return FText::FromString(FString::Printf(TEXT("%02i"), GetCurrentDateTime().GetMonth()));
+	    })
+	});
+
+	Tokens.Add({
+	    TEXT("year"),
+	    LOCTEXT("TokenYear", "Year"),
+	    FNamingTokenData::FTokenProcessorDelegateNative::CreateLambda([this]{
+	        return FText::FromString(FString::Printf(TEXT("%04i"), GetCurrentDateTime().GetYear()));
+	    })
+	});
+
+	Tokens.Add({
+	    TEXT("hour"),
+	    LOCTEXT("TokenHour", "Hour"),
+	    FNamingTokenData::FTokenProcessorDelegateNative::CreateLambda([this]{
+	        return FText::FromString(FString::Printf(TEXT("%02i"), GetCurrentDateTime().GetHour()));
+	    })
+	});
+
+	Tokens.Add({
+	    TEXT("minute"),
+	    LOCTEXT("TokenMinute", "Minute"),
+	    FNamingTokenData::FTokenProcessorDelegateNative::CreateLambda([this]{
+	        return FText::FromString(FString::Printf(TEXT("%02i"), GetCurrentDateTime().GetMinute()));
+	    })
+	});
+
+	Tokens.Add({
+	    TEXT("second"),
+	    LOCTEXT("TokenSecond", "Second"),
+	    FNamingTokenData::FTokenProcessorDelegateNative::CreateLambda([this]{
+	        return FText::FromString(FString::Printf(TEXT("%02i"), GetCurrentDateTime().GetSecond()));
+	    })
+	});
+
+	Tokens.Add({
+	    TEXT("take"),
+	    LOCTEXT("TokenTake", "Take"),
+	    FNamingTokenData::FTokenProcessorDelegateNative::CreateLambda([this]{
+	        const UMovieSceneToolsProjectSettings* ProjectSettings = GetDefault<UMovieSceneToolsProjectSettings>();
+	        const int32 TakeNumDigits = ProjectSettings->TakeNumDigits;
+	        return FText::FromString(FString::Printf(TEXT("%0*d"), TakeNumDigits, TakeMetaData->GetTakeNumber()));
+	    })
+	});
+
+	Tokens.Add({
+	    TEXT("slate"),
+	    LOCTEXT("TokenSlate", "Slate"),
+	    FNamingTokenData::FTokenProcessorDelegateNative::CreateLambda([this]{
+	        return FText::FromString(TakeMetaData->GetSlate());
+	    })
+	});
+
+	Tokens.Add({
+	    TEXT("map"),
+	    LOCTEXT("TokenMap", "Map"),
+	    FNamingTokenData::FTokenProcessorDelegateNative::CreateLambda([]{
+	        FString MapName;
+	        if (GIsEditor)
+	        {
+	            MapName = FPackageName::GetShortFName(
+	                GEditor->GetEditorWorldContext().World()->PersistentLevel->GetOutermost()->GetFName()
+	            ).GetPlainNameString();
+	        }
+	        return FText::FromString(MapName);
+	    })
+	});
+}
+
+void UTakeRecorderNamingTokens::OnPreEvaluate_Implementation(const FNamingTokensEvaluationData& InEvaluationData)
+{
+	Super::OnPreEvaluate_Implementation(InEvaluationData);
+	
+	TakeMetaData = UTakeMetaData::GetMostRecentMetaData();
+	if (TakeMetaData == nullptr)
+	{
+		// Create metadata which computes its data from available information.
+		// This is how STakeRecorderCockpit handles cases when metadata isn't available. Ideally this
+		// method would be available for both us and the slate widget, but UTakeRecorderProjectSettings creates
+		// a circular dependency we should avoid.
+		
+		TakeMetaData = UTakeMetaData::CreateFromDefaults(GetTransientPackage(), NAME_None);
+		TakeMetaData->SetFlags(RF_Transactional | RF_Transient);
+
+		const FString DefaultSlate = GetDefault<UTakeRecorderProjectSettings>()->Settings.DefaultSlate;
+		if (TakeMetaData->GetSlate() != DefaultSlate)
+		{
+			TakeMetaData->SetSlate(DefaultSlate, false);
+		}
+
+		// Compute the correct starting take number
+		const int32 NextTakeNumber = UTakesCoreBlueprintLibrary::ComputeNextTakeNumber(TakeMetaData->GetSlate());
+		if (TakeMetaData->GetTakeNumber() != NextTakeNumber)
+		{
+			TakeMetaData->SetTakeNumber(NextTakeNumber, false);
+		}
+	}
+}
+
+void UTakeRecorderNamingTokens::OnPostEvaluate_Implementation()
+{
+	Super::OnPostEvaluate_Implementation();
+
+	// Make sure we don't keep a strong reference in case we're linking to active metadata.
+	TakeMetaData = nullptr;
+}
+
+FDateTime UTakeRecorderNamingTokens::GetCurrentDateTime_Implementation() const
+{
+	const FDateTime TakeTimestamp = TakeMetaData ? TakeMetaData->GetTimestamp() : FDateTime(0);
+	return TakeTimestamp == FDateTime(0) ? Super::GetCurrentDateTime_Implementation() : TakeTimestamp;
+}
+
+#undef LOCTEXT_NAMESPACE
