Take Recorder: Support dynamic user tokens. Users can enter in tokens into specific fields in Take Recorder, auto generating new tokens in the top UI. These exist for the life cycle of the panel and allows users to define their own tokens to use within these fields.

Fix TakeRecorder support for actor and channel tokens. Fix issue where externally calling a TR token wouldn't have up-to-date information once a take was recorded.

NamingTokens: Support registering external instance tokens.
Add support for passing in a context UObject to token evaluation.
Change token processing to only process token keys found in the string, rather than try every token in a namespace.
Individual token results are now included in the final result, in order of their appearance.
Add Pre and Post evaluate delegates to NamingToken classes.

Live Link Hub: Make Naming Tokens namespace public.

diff --git a/Engine/Plugins/Animation/LiveLinkHub/Source/LiveLinkHub/Private/Config/LiveLinkHubTemplateTokens.cpp b/Engine/Plugins/Animation/LiveLinkHub/Source/LiveLinkHub/Private/Config/LiveLinkHubTemplateTokens.cpp
--- a/Engine/Plugins/Animation/LiveLinkHub/Source/LiveLinkHub/Private/Config/LiveLinkHubTemplateTokens.cpp
+++ b/Engine/Plugins/Animation/LiveLinkHub/Source/LiveLinkHub/Private/Config/LiveLinkHubTemplateTokens.cpp
@@ -2,13 +2,14 @@
 
 #include "LiveLinkHubTemplateTokens.h"
 
+#include "ILiveLinkHubModule.h"
 #include "Session/LiveLinkHubSessionManager.h"
 
 #define LOCTEXT_NAMESPACE "LiveLinkHubNamingTokens"
 
 ULiveLinkHubNamingTokens::ULiveLinkHubNamingTokens()
 {
-	Namespace = TEXT("llh");
+	Namespace = ILiveLinkHubModule::GetLiveLinkHubNamingTokensNamespace();
 }
 
 void ULiveLinkHubNamingTokens::OnCreateDefaultTokens(TArray<FNamingTokenData>& Tokens)
diff --git a/Engine/Plugins/Animation/LiveLinkHub/Source/LiveLinkHub/Public/ILiveLinkHubModule.h b/Engine/Plugins/Animation/LiveLinkHub/Source/LiveLinkHub/Public/ILiveLinkHubModule.h
--- a/Engine/Plugins/Animation/LiveLinkHub/Source/LiveLinkHub/Public/ILiveLinkHubModule.h
+++ b/Engine/Plugins/Animation/LiveLinkHub/Source/LiveLinkHub/Public/ILiveLinkHubModule.h
@@ -2,6 +2,7 @@
 
 #pragma once
 
+#include "Containers/UnrealString.h"
 #include "Modules/ModuleInterface.h"
 
 class LIVELINKHUB_API ILiveLinkHubModule : public IModuleInterface
@@ -16,4 +17,10 @@ public:
 
 	UE_DEPRECATED(5.6, "There's no need to call this method anymore, it's now handled by StartupModule and ShutdownModule.")
 	virtual void ShutdownLiveLinkHub() {}
+
+	/** Retrieve the namespace for Live Link Hub naming tokens. */
+	static FString GetLiveLinkHubNamingTokensNamespace()
+	{
+		return TEXT("llh");
+	}
 };
diff --git a/Engine/Plugins/Developer/NamingTokens/Source/NamingTokens/Private/NamingTokens.cpp b/Engine/Plugins/Developer/NamingTokens/Source/NamingTokens/Private/NamingTokens.cpp
--- a/Engine/Plugins/Developer/NamingTokens/Source/NamingTokens/Private/NamingTokens.cpp
+++ b/Engine/Plugins/Developer/NamingTokens/Source/NamingTokens/Private/NamingTokens.cpp
@@ -7,6 +7,7 @@
 #include "Utils/NamingTokenUtils.h"
 
 #include "Engine/Engine.h"
+#include "Internationalization/LocKeyFuncs.h"
 #include "Internationalization/Regex.h"
 #include "UObject/Package.h"
 
@@ -49,9 +50,11 @@ void UNamingTokens::OnCreateDefaultTokens(TArray<FNamingTokenData>& Tokens)
 {
 }
 
-FNamingTokenResultData UNamingTokens::EvaluateTokenText(const FText& InTokenText)
+FNamingTokenResultData UNamingTokens::EvaluateTokenText(const FText& InTokenText, const TArray<UObject*>& InContexts)
 {
-	return EvaluateTokenText(InTokenText, FNamingTokensEvaluationData());
+	FNamingTokensEvaluationData EvaluationData;
+	EvaluationData.Contexts = InContexts;
+	return EvaluateTokenText(InTokenText, EvaluationData);
 }
 
 FNamingTokenResultData UNamingTokens::EvaluateTokenText(const FText& InTokenText, const FNamingTokensEvaluationData& InEvaluationData)
@@ -62,17 +65,71 @@ FNamingTokenResultData UNamingTokens::EvaluateTokenText(const FText& InTokenText
 	Result.OriginalText = InTokenText;
 
 	FText FormattedText = InTokenText;
-	{
-		for (const FNamingTokenData& Token : GetDefaultTokens())
-		{
-			TryProcessToken(Token, FormattedText);
-		}
 
-		for (const FNamingTokenData& Token : CustomTokens)
+	// Retrieve all unprocessed tokens within this text.
+	const TArray<FString> UnprocessedTokens = UE::NamingTokens::Utils::Private::GetTokenKeysFromString(FormattedText.ToString());
+
+	TSet<FNamingTokenData> DefaultTokensSet(GetDefaultTokens());
+	TSet<FNamingTokenData> CustomTokenSet(CustomTokens);
+
+	TSet<FString, FLocKeySetFuncs> CheckedKeys;
+	
+	for (const FString& TokenKeyWithNamespace : UnprocessedTokens)
+	{
+		const FString TokenNamespace = UE::NamingTokens::Utils::Private::GetNamespaceFromTokenKey(TokenKeyWithNamespace);
+		if (TokenNamespace.IsEmpty() || TokenNamespace == Namespace)
 		{
-			TryProcessToken(Token, FormattedText);
+			// Make sure the token key doesn't include the namespace.
+			const FString TokenKey = UE::NamingTokens::Utils::Private::RemoveNamespaceFromTokenKey(TokenKeyWithNamespace);
+
+			if (CheckedKeys.Contains(TokenKey))
+			{
+				continue;
+			}
+			CheckedKeys.Add(TokenKey);
+			
+			FNamingTokenValueData ValueData;
+			
+			FText TokenValue;
+			const FNamingTokenData* SelectedToken = nullptr;
+			if (FNamingTokenData* DefaultToken = DefaultTokensSet.Find(FNamingTokenData(TokenKey)))
+			{
+				SelectedToken = DefaultToken;
+			}
+			else if (FNamingTokenData* CustomToken = CustomTokenSet.Find(FNamingTokenData(TokenKey)))
+			{
+				SelectedToken = CustomToken;
+			}
+			else
+			{
+				// Not found in a class definition, check external tokens.
+				for (const TTuple<FGuid, TArray<FNamingTokenData>>& ExternalTokenKeyPair : ExternalTokens)
+				{
+					if (const FNamingTokenData* ExternalToken = ExternalTokenKeyPair.Value.FindByKey(FNamingTokenData(TokenKey)))
+					{
+						SelectedToken = ExternalToken;
+						break;
+					}
+				}
+			}
+			
+			if (SelectedToken)
+			{
+				TokenValue = ProcessToken(*SelectedToken, FormattedText);
+				ValueData.bWasEvaluated = true;
+			}
+			else
+			{
+				ValueData.bWasEvaluated = false;
+			}
+			
+			ValueData.TokenKey = TokenKey;
+			ValueData.TokenValue = TokenValue;
+			ValueData.ProvidedNamespace = TokenNamespace;
+			Result.TokenValues.Add(MoveTemp(ValueData));
 		}
 	}
+
 	Result.EvaluatedText = FormattedText;
 	
 	// Split folder path and file name.
@@ -93,7 +150,7 @@ FNamingTokenResultData UNamingTokens::EvaluateTokenText(const FText& InTokenText
 		Result.FileName = FormattedString;
 	}
 
-	PostEvaluate();
+	PostEvaluate(Result);
 	
 	return Result;
 }
@@ -105,13 +162,13 @@ FString UNamingTokens::GetFormattedTokensStringForDisplay() const
 	for (const FNamingTokenData& Token : GetDefaultTokens())
 	{
 		FormattedTokensString += FString::Printf(TEXT("%s - %s\n"),
-			*CreateFormattedToken(Token), *Token.DisplayName.ToString());
+			*UE::NamingTokens::Utils::Private::CreateFormattedToken(Token), *Token.DisplayName.ToString());
 	}
 
 	for (const FNamingTokenData& Token : CustomTokens)
 	{
 		FormattedTokensString += FString::Printf(TEXT("%s - %s\n"),
-			*CreateFormattedToken(Token), *Token.DisplayName.ToString());
+			*UE::NamingTokens::Utils::Private::CreateFormattedToken(Token), *Token.DisplayName.ToString());
 	}
 	
 	return FormattedTokensString;
@@ -128,12 +185,37 @@ const TArray<FNamingTokenData>& UNamingTokens::GetCustomTokens() const
 	return CustomTokens;
 }
 
+TArray<FNamingTokenData>& UNamingTokens::RegisterExternalTokens(FGuid& OutGuid)
+{
+	OutGuid = FGuid::NewGuid();
+	return ExternalTokens.Add(OutGuid);
+}
+
+void UNamingTokens::UnregisterExternalTokens(const FGuid& InGuid)
+{
+	ExternalTokens.Remove(InGuid);
+}
+
+bool UNamingTokens::AreExternalTokensRegistered(const FGuid& InGuid) const
+{
+	return ExternalTokens.Contains(InGuid);
+}
+
+TArray<FNamingTokenData>& UNamingTokens::GetExternalTokensChecked(const FGuid& InGuid)
+{
+	return ExternalTokens.FindChecked(InGuid);
+}
+
 TArray<FNamingTokenData> UNamingTokens::GetAllTokens() const
 {
 	TArray<FNamingTokenData> AllTokens;
 	AllTokens.Reserve(DefaultTokens.Num() + CustomTokens.Num());
 	AllTokens.Append(DefaultTokens);
 	AllTokens.Append(CustomTokens);
+	for (const TTuple<FGuid, TArray<FNamingTokenData>>& ExternalTokenKeyPair : ExternalTokens)
+	{
+		AllTokens.Append(ExternalTokenKeyPair.Value);
+	}
 	return AllTokens;
 }
 
@@ -146,11 +228,13 @@ void UNamingTokens::PreEvaluate(const FNamingTokensEvaluationData& InEvaluationD
 {
 	CurrentEvaluationData = InEvaluationData;
 	OnPreEvaluate(InEvaluationData);
+	OnPreEvaluateEvent.Broadcast(InEvaluationData);
 }
 
-void UNamingTokens::PostEvaluate()
+void UNamingTokens::PostEvaluate(const FNamingTokenResultData& InResultData)
 {
 	OnPostEvaluate();
+	OnPostEvaluateEvent.Broadcast(InResultData);
 }
 
 void UNamingTokens::OnPreEvaluate_Implementation(const FNamingTokensEvaluationData& InEvaluationData)
@@ -161,9 +245,9 @@ void UNamingTokens::OnPostEvaluate_Implementation()
 {
 }
 
-bool UNamingTokens::TryProcessToken(const FNamingTokenData& InToken, FText& InOutFormattedText)
+FText UNamingTokens::ProcessToken(const FNamingTokenData& InToken, FText& InOutFormattedText)
 {
-	const FString FormattedToken = *CreateFormattedToken(InToken);
+	const FString FormattedToken = *UE::NamingTokens::Utils::Private::CreateFormattedToken(InToken);
 
 	FText ProcessedToken = FText::GetEmpty();
 
@@ -183,7 +267,7 @@ bool UNamingTokens::TryProcessToken(const FNamingTokenData& InToken, FText& InOu
 	else
 	{
 		UE_LOG(LogNamingTokens, Error, TEXT("Could not find function process for token %s"), *FormattedToken);
-		return false;
+		return ProcessedToken;
 	}
 		
 	FFormatNamedArguments Args;
@@ -191,12 +275,7 @@ bool UNamingTokens::TryProcessToken(const FNamingTokenData& InToken, FText& InOu
 	Args.Add(FString::Printf(TEXT("%s.%s"), *GetNamespace(), *InToken.TokenKey), ProcessedToken);
 
 	InOutFormattedText = FText::Format(InOutFormattedText, Args);
-	return true;
-}
-
-FString UNamingTokens::CreateFormattedToken(const FNamingTokenData& InToken) const
-{
-	return FString::Printf(TEXT("{%s}"), *InToken.TokenKey);
+	return ProcessedToken;
 }
 
 UFunction* UNamingTokens::FindBlueprintFunctionForToken(const FNamingTokenData& InTokenData) const
diff --git a/Engine/Plugins/Developer/NamingTokens/Source/NamingTokens/Private/NamingTokensEditorSubsystem.cpp b/Engine/Plugins/Developer/NamingTokens/Source/NamingTokens/Private/NamingTokensEditorSubsystem.cpp
--- a/Engine/Plugins/Developer/NamingTokens/Source/NamingTokens/Private/NamingTokensEditorSubsystem.cpp
+++ b/Engine/Plugins/Developer/NamingTokens/Source/NamingTokens/Private/NamingTokensEditorSubsystem.cpp
@@ -5,6 +5,7 @@
 #include "NamingTokens.h"
 #include "NamingTokensEvaluationData.h"
 #include "NamingTokensLog.h"
+#include "Utils/NamingTokenUtils.h"
 
 #include "AssetRegistry/AssetData.h"
 #include "AssetRegistry/AssetDataTagMap.h"
@@ -21,7 +22,7 @@ TArray<FString> UNamingTokensEditorSubsystem::GetAllNamespaces() const
 	// Find all assets whose native parent class is UNamingTokens, and add their namespaces
 	TArray<FAssetData> NamingTokenAssets;
 	const TMultiMap<FName, FString> TagValues = { { FBlueprintTags::NativeParentClassPath, FObjectPropertyBase::GetExportPath(UNamingTokens::StaticClass()) } };
-	AssetRegistryModule.Get().GetAssetsByTagValues(TagValues, NamingTokenAssets);
+	AssetRegistry.GetAssetsByTagValues(TagValues, NamingTokenAssets);
 
 	TSet<FString> DiscoveredNamespaces;
 	for (const FAssetData& NamingTokenAssetData : NamingTokenAssets)
@@ -190,7 +191,7 @@ TArray<UNamingTokens*> UNamingTokensEditorSubsystem::GetMultipleNamingTokens(con
 	return Result;
 }
 
-FNamingTokenResultData UNamingTokensEditorSubsystem::EvaluateTokenText(const FText& InTokenText, FNamingTokenFilterArgs InFilter)
+FNamingTokenResultData UNamingTokensEditorSubsystem::EvaluateTokenText(const FText& InTokenText, const FNamingTokenFilterArgs& InFilter, const TArray<UObject*>& InContexts)
 {
 	FNamingTokenResultData Result;
 
@@ -199,6 +200,8 @@ FNamingTokenResultData UNamingTokensEditorSubsystem::EvaluateTokenText(const FTe
 
 	TSet<FString> Namespaces = GetNamingTokenNamespacesFromString(InTokenText.ToString(), InFilter);
 
+	TArray<FNamingTokenValueData> CompletedNamingTokenValueData;
+	
 	// These external filters are intentionally executed before adding any namespaces from the input FNamingTokenFilterArgs filter. 
 	// This allows the caller of this function to guarantee that the namespaces they want included will always be allowed, regardless of other filters that may execute.
 	for (const TPair<FName, FFilterNamespace>& DelegatePair : FilterNamespaceDelegates)
@@ -208,24 +211,62 @@ FNamingTokenResultData UNamingTokensEditorSubsystem::EvaluateTokenText(const FTe
 
 	Namespaces.Append(InFilter.NamespacesToInclude); // Tokens could have been written without the namespace but should still be scoped to the filter.
 
-	const FNamingTokensEvaluationData EvaluationData;
+	FNamingTokensEvaluationData EvaluationData;
+	EvaluationData.Contexts = InContexts;
 	for (const FString& Namespace : Namespaces)
 	{
 		if (UNamingTokens* NamingTokens = GetNamingTokens(Namespace))
 		{
 			Result = NamingTokens->EvaluateTokenText(ProcessedTokenText, EvaluationData);
+
+			// Properly record individual token evaluation data. The order should be consistent with the order
+			// a key was processed in the string. A key may be processed multiple times if it is referenced multiple times, or is not
+			// identified in one NamingTokens class but is in another. In the event a future class identifies it, it will update the original
+			// key.
+			for (const FNamingTokenValueData& NewTokenValue : Result.TokenValues)
+			{
+				bool bReplacedOldValue = false;
+				// Find previously recorded values that were undefined and update them indicating they are no longer undefined.
+				for (FNamingTokenValueData& OldTokenValue : CompletedNamingTokenValueData)
+				{
+					if (OldTokenValue.TokenKey.Equals(NewTokenValue.TokenKey, ESearchCase::CaseSensitive))
+					{
+						OldTokenValue = NewTokenValue;
+						bReplacedOldValue = true;
+						break;
+					}
+				}
+				if (!bReplacedOldValue)
+				{
+					// Add to the completed naming token value.
+					CompletedNamingTokenValueData.Add(NewTokenValue);
+				}
+			}
+			
 			ProcessedTokenText = Result.EvaluatedText;
 		}
 	}
 
 	Result.OriginalText = InTokenText;
+
+	Result.TokenValues = MoveTemp(CompletedNamingTokenValueData);
 	
 	return Result;
 }
 
-FNamingTokenResultData UNamingTokensEditorSubsystem::EvaluateTokenString(const FString& InTokenString, FNamingTokenFilterArgs InFilter)
+FNamingTokenResultData UNamingTokensEditorSubsystem::EvaluateTokenText(const FText& InTokenText, const FNamingTokenFilterArgs& InFilter)
+{
+	return EvaluateTokenText(InTokenText, InFilter, {});
+}
+
+FNamingTokenResultData UNamingTokensEditorSubsystem::EvaluateTokenString(const FString& InTokenString, const FNamingTokenFilterArgs& InFilter, const TArray<UObject*>& InContexts)
+{
+	return EvaluateTokenText(FText::FromString(InTokenString), InFilter, InContexts);
+}
+
+FNamingTokenResultData UNamingTokensEditorSubsystem::EvaluateTokenString(const FString& InTokenString, const FNamingTokenFilterArgs& InFilter)
 {
-	return EvaluateTokenText(FText::FromString(InTokenString), InFilter);
+	return EvaluateTokenString(InTokenString, InFilter, {});
 }
 
 void UNamingTokensEditorSubsystem::RegisterGlobalNamespace(const FString& InNamespace)
diff --git a/Engine/Plugins/Developer/NamingTokens/Source/NamingTokens/Private/Utils/NamingTokenUtils.cpp b/Engine/Plugins/Developer/NamingTokens/Source/NamingTokens/Private/Utils/NamingTokenUtils.cpp
--- a/Engine/Plugins/Developer/NamingTokens/Source/NamingTokens/Private/Utils/NamingTokenUtils.cpp
+++ b/Engine/Plugins/Developer/NamingTokens/Source/NamingTokens/Private/Utils/NamingTokenUtils.cpp
@@ -4,6 +4,8 @@
 
 #include "NamingTokens.h"
 
+#include "Internationalization/LocKeyFuncs.h"
+#include "Internationalization/Regex.h"
 #include "Kismet2/BlueprintEditorUtils.h"
 #include "Kismet2/KismetEditorUtilities.h"
 
@@ -48,3 +50,44 @@ FName UE::NamingTokens::Utils::Private::CreateNewTokenGraph(UBlueprint* InBluepr
 
 	return FunctionName;
 }
+
+FString UE::NamingTokens::Utils::Private::CreateFormattedToken(const FNamingTokenData& InToken)
+{
+	return FString::Printf(TEXT("{%s}"), *InToken.TokenKey);
+}
+
+TArray<FString> UE::NamingTokens::Utils::Private::GetTokenKeysFromString(const FString& InTokenString)
+{
+	TSet<FString, FLocKeySetFuncs> Tokens;
+
+	const FString PatternString = TEXT(R"(\{\s*((?:[a-zA-Z0-9_]+\.)*[a-zA-Z0-9_]+)\s*\})");
+	const FRegexPattern Pattern(PatternString);
+	FRegexMatcher Matcher(Pattern, InTokenString);
+
+	while (Matcher.FindNext())
+	{
+		FString Token = Matcher.GetCaptureGroup(1);
+		Tokens.Add(Token);
+	}
+
+	return Tokens.Array();
+}
+
+FString UE::NamingTokens::Utils::Private::GetNamespaceFromTokenKey(const FString& InTokenKey)
+{
+	FString LeftPart, RightPart;
+	InTokenKey.Split(TEXT("."), &LeftPart, &RightPart);
+
+	return LeftPart;
+}
+
+FString UE::NamingTokens::Utils::Private::RemoveNamespaceFromTokenKey(const FString& InTokenKey)
+{
+	FString Result = InTokenKey;
+	FString LeftPart, RightPart;
+	if (InTokenKey.Split(TEXT("."), &LeftPart, &RightPart))
+	{
+		Result = RightPart;
+	}
+	return Result;
+}
diff --git a/Engine/Plugins/Developer/NamingTokens/Source/NamingTokens/Private/Utils/NamingTokenUtils.h b/Engine/Plugins/Developer/NamingTokens/Source/NamingTokens/Private/Utils/NamingTokenUtils.h
--- a/Engine/Plugins/Developer/NamingTokens/Source/NamingTokens/Private/Utils/NamingTokenUtils.h
+++ b/Engine/Plugins/Developer/NamingTokens/Source/NamingTokens/Private/Utils/NamingTokenUtils.h
@@ -5,6 +5,7 @@
 #include "CoreMinimal.h"
 
 class UBlueprint;
+struct FNamingTokenData;
 
 namespace UE::NamingTokens::Utils::Private
 {
@@ -19,4 +20,24 @@ namespace UE::NamingTokens::Utils::Private
 
 	/** Creates a new k2 graph in a blueprint for a given token key. */
 	FName CreateNewTokenGraph(UBlueprint* InBlueprint, const FString& InTokenKey);
+
+	/** Formats a token key into {key}. */
+	FString CreateFormattedToken(const FNamingTokenData& InToken);
+
+	/** Extract all token keys from a string. */
+	TArray<FString> GetTokenKeysFromString(const FString& InTokenString);
+
+	/**
+	 * Return the namespace from a token key, or an empty string.
+	 * @param InTokenKey The token key shouldn't have any additional formatting beyond the namespace, eg "namespace.token"
+	 * @return The namespace, or an empty string.
+	 */
+	FString GetNamespaceFromTokenKey(const FString& InTokenKey);
+	
+	/**
+	 * Remove the namespace from the given token key.
+	 * @param InTokenKey The token key shouldn't have any additional formatting beyond the namespace, eg "namespace.token"
+	 * @return The token key without the namespace, or the InTokenKey if there was no namespace.
+	 */
+	FString RemoveNamespaceFromTokenKey(const FString& InTokenKey);
 }
\ No newline at end of file
diff --git a/Engine/Plugins/Developer/NamingTokens/Source/NamingTokens/Public/NamingTokenData.h b/Engine/Plugins/Developer/NamingTokens/Source/NamingTokens/Public/NamingTokenData.h
--- a/Engine/Plugins/Developer/NamingTokens/Source/NamingTokens/Public/NamingTokenData.h
+++ b/Engine/Plugins/Developer/NamingTokens/Source/NamingTokens/Public/NamingTokenData.h
@@ -7,6 +7,27 @@
 
 #include "NamingTokenData.generated.h"
 
+class UNamingTokens;
+
+USTRUCT(BlueprintType)
+struct FNamingTokenValueData
+{
+	GENERATED_BODY()
+	
+	/** The token key. */
+	UPROPERTY(BlueprintReadOnly, Category = "NamingTokens")
+	FString TokenKey;
+	/** The namespace of the token, if it was originally provided during evaluation. */
+	UPROPERTY(BlueprintReadOnly, Category = "NamingTokens")
+	FString ProvidedNamespace;
+	/** The evaluated token text. */
+	UPROPERTY(BlueprintReadOnly, Category = "NamingTokens")
+	FText TokenValue;
+	/** If the token was able to be evaluated. */
+	UPROPERTY(BlueprintReadOnly, Category = "NamingTokens")
+	bool bWasEvaluated = false;
+};
+
 /** Evaluated results from a template file string. */
 USTRUCT(BlueprintType)
 struct NAMINGTOKENS_API FNamingTokenResultData
@@ -19,12 +40,18 @@ struct NAMINGTOKENS_API FNamingTokenResultData
 	/** The full text with evaluated tokens. */
 	UPROPERTY(BlueprintReadOnly, Category = "NamingTokens")
 	FText EvaluatedText;
-	/** The evaluated folder(s) without the first forward slash or filename. */
+	/** @todo NamingTokens - Remove and leave up to caller to manage.
+	 * The evaluated folder(s) without the first forward slash or filename. */
 	UPROPERTY(BlueprintReadOnly, Category = "NamingTokens")
 	FString FolderPath;
-	/** The evaluated filename. */
+	/** @todo NamingTokens - Remove and leave up to caller to manage.
+	 * The evaluated filename. */
 	UPROPERTY(BlueprintReadOnly, Category = "NamingTokens")
 	FString FileName;
+
+	/** The result of individual tokens, in the order they appear in OriginalText. */
+	UPROPERTY(BlueprintReadOnly, Category = "NamingTokens")
+	TArray<FNamingTokenValueData> TokenValues;
 };
 
 USTRUCT(BlueprintType)
@@ -36,6 +63,12 @@ public:
 	FNamingTokenData() = default;
 
 	DECLARE_DELEGATE_RetVal(FText, FTokenProcessorDelegateNative);
+
+	explicit FNamingTokenData(const FString& InTokenKey)
+	{
+		TokenKey = InTokenKey;
+	}
+	
 	FNamingTokenData(const FString& InTokenKey, const FText& InTokenDisplayName, const FTokenProcessorDelegateNative& InTokenProcessor)
 	{
 		DisplayName = InTokenDisplayName;
@@ -68,4 +101,15 @@ public:
 
 	/** The native delegate to execute to evaluate the function. If FunctionName is set then this is not used. */
 	FTokenProcessorDelegateNative TokenProcessorNative;
+
+	bool operator==(const FNamingTokenData& Other) const
+	{
+		return TokenKey.Equals(Other.TokenKey, ESearchCase::CaseSensitive);
+	}
+	
+	friend FORCEINLINE uint32 GetTypeHash(const FNamingTokenData& Item)
+	{
+		const uint32 Hash = FCrc::StrCrc32(*Item.TokenKey);
+		return Hash;
+	}
 };
diff --git a/Engine/Plugins/Developer/NamingTokens/Source/NamingTokens/Public/NamingTokens.h b/Engine/Plugins/Developer/NamingTokens/Source/NamingTokens/Public/NamingTokens.h
--- a/Engine/Plugins/Developer/NamingTokens/Source/NamingTokens/Public/NamingTokens.h
+++ b/Engine/Plugins/Developer/NamingTokens/Source/NamingTokens/Public/NamingTokens.h
@@ -36,7 +36,7 @@ protected:
 
 public:
 	/** Evaluate token text. Creates EvaluationData. */
-	FNamingTokenResultData EvaluateTokenText(const FText& InTokenText);
+	FNamingTokenResultData EvaluateTokenText(const FText& InTokenText, const TArray<UObject*>& InContexts = {});
 
 	/** Evaluate token text for a given EvaluationData. */
 	FNamingTokenResultData EvaluateTokenText(const FText& InTokenText, const FNamingTokensEvaluationData& InEvaluationData);
@@ -49,6 +49,39 @@ public:
 	
 	/** Retrieve the custom tokens. */
 	const TArray<FNamingTokenData>& GetCustomTokens() const;
+
+	/**
+	 * Register an external token array which is managed from a caller. This is transient data.
+	 * To retrieve the array, call GetExternalTokensChecked and pass in the Guid.
+	 *
+	 * @param OutGuid The new guid these tokens are registered under.
+	 *
+	 * @return A reference to the token array. This reference may be invalidated if another array is registered.
+	 */
+	TArray<FNamingTokenData>& RegisterExternalTokens(FGuid& OutGuid);
+
+	/**
+	 * Unregister and clear out external tokens.
+	 * 
+	 * @param InGuid The guid to unregister.
+	 */
+	void UnregisterExternalTokens(const FGuid& InGuid);
+	
+	/**
+	 * Check if external tokens are registered for a guid.
+	 *
+	 * @param InGuid The guid which was registered for these external tokens.
+	 * @return True if the tokens are registered, false if they are not registered.
+	 */
+	bool AreExternalTokensRegistered(const FGuid& InGuid) const;
+
+	/**
+	 * Retrieve the external tokens. Tokens must be registered and exist.
+	 *
+	 * @param InGuid The guid which was registered for these external tokens.
+	 * @return A reference to the token array. This reference may be invalidated if another array is registered.
+	 */
+	TArray<FNamingTokenData>& GetExternalTokensChecked(const FGuid& InGuid);
 	
 	/** Retrieve all tokens. */
 	TArray<FNamingTokenData> GetAllTokens() const;
@@ -57,7 +90,13 @@ public:
 	const FString& GetNamespace() const { return Namespace; }
 	/** The name of the namespace property. */
 	static FName GetNamespacePropertyName() { return GET_MEMBER_NAME_CHECKED(UNamingTokens, Namespace); }
-
+	
+	/** Delegate when pre evaluate is called. */
+	DECLARE_MULTICAST_DELEGATE_OneParam(FOnPreEvaluate, const FNamingTokensEvaluationData& InEvaluationData);
+	FOnPreEvaluate& GetOnPreEvaluateEvent() { return OnPreEvaluateEvent; }
+	/** Delegate when post evaluate is called. */
+	DECLARE_MULTICAST_DELEGATE_OneParam(FOnPostEvaluate, const FNamingTokenResultData& InResultData);
+	FOnPostEvaluate& GetOnPostEvaluateEvent() { return OnPostEvaluateEvent; }
 protected:
 	/** Retrieve the current datetime. By default, this uses shared data so results are consistent across runs. */
 	UFUNCTION(BlueprintCallable, BlueprintNativeEvent, Category = "NamingTokens")
@@ -84,7 +123,7 @@ protected:
 	/**
 	 * Called after all tokens have evaluated.
 	 */
-	void PostEvaluate();
+	void PostEvaluate(const FNamingTokenResultData& InResultData);
 	
 	/**
 	 * Called prior to evaluation. Allows consistent data to be set up for each token evaluation.
@@ -104,11 +143,13 @@ protected:
 	virtual void OnPostEvaluate_Implementation();
 
 private:
+	/** Delegate when pre evaluate is called. */
+	FOnPreEvaluate OnPreEvaluateEvent;
+	/** Delegate when post evaluate is called. */
+	FOnPostEvaluate OnPostEvaluateEvent;
+	
 	/** Process a token if is defined and return evaluated text. */
-	bool TryProcessToken(const FNamingTokenData& InToken, FText& InOutFormattedText);
-
-	/** Formats a token key into {key}. */
-	FString CreateFormattedToken(const FNamingTokenData& InToken) const;
+	FText ProcessToken(const FNamingTokenData& InToken, FText& InOutFormattedText);
 	
 	/** Return a blueprint function for a token processor if it exists. */
 	UFunction* FindBlueprintFunctionForToken(const FNamingTokenData& InTokenData) const;
@@ -117,6 +158,9 @@ private:
 	/** The default tokens defined by this class. */
 	TArray<FNamingTokenData> DefaultTokens;
 
+	/** External and temporary instance tokens which can be filled in by tools supporting unrecognized tokens. */
+	TMap<FGuid, TArray<FNamingTokenData>> ExternalTokens;
+	
 protected:
 	/** User defined tokens. */
 	UPROPERTY(EditDefaultsOnly, Category = "NamingTokens", meta = (TitleProperty="TokenKey"))
diff --git a/Engine/Plugins/Developer/NamingTokens/Source/NamingTokens/Public/NamingTokensEditorSubsystem.h b/Engine/Plugins/Developer/NamingTokens/Source/NamingTokens/Public/NamingTokensEditorSubsystem.h
--- a/Engine/Plugins/Developer/NamingTokens/Source/NamingTokens/Public/NamingTokensEditorSubsystem.h
+++ b/Engine/Plugins/Developer/NamingTokens/Source/NamingTokens/Public/NamingTokensEditorSubsystem.h
@@ -57,14 +57,32 @@ public:
 	UFUNCTION(BlueprintCallable, Category = "NamingTokens")
 	TArray<UNamingTokens*> GetMultipleNamingTokens(const TArray<FString>& InNamespaces) const;
 
-	/** Parse and evaluate token text. */
-	UFUNCTION(BlueprintCallable, Category = "NamingTokens")
-	FNamingTokenResultData EvaluateTokenText(const FText& InTokenText, FNamingTokenFilterArgs InFilter);
-
-	/** Parse and evaluate a token string. */
-	UFUNCTION(BlueprintCallable, Category = "NamingTokens")
-	FNamingTokenResultData EvaluateTokenString(const FString& InTokenString, FNamingTokenFilterArgs InFilter);
-
+	/**
+	 * Parse and evaluate token text.
+	 * 
+	 * @param InTokenText The text containing unprocessed tokens.
+	 * @param InFilter [Optional] Filter to determine which namespaces to use.
+	 * @param InContexts [Optional] Context objects to pass to naming tokens.
+	 *
+	 * @return The result of the evaluation.
+	 */
+	UFUNCTION(BlueprintCallable, Category = "NamingTokens", meta = (AutoCreateRefTerm = "InFilter,InContexts"))
+	FNamingTokenResultData EvaluateTokenText(const FText& InTokenText, const FNamingTokenFilterArgs& InFilter, const TArray<UObject*>& InContexts);
+	FNamingTokenResultData EvaluateTokenText(const FText& InTokenText, const FNamingTokenFilterArgs& InFilter = FNamingTokenFilterArgs());
+	
+	/**
+	 * Parse and evaluate token string.
+	 * 
+	 * @param InTokenString The string containing unprocessed tokens.
+	 * @param InFilter [Optional] Filter to determine which namespaces to use.
+	 * @param InContexts [Optional] Context objects to pass to naming tokens.
+	 *
+	 * @return The result of the evaluation.
+	 */
+	UFUNCTION(BlueprintCallable, Category = "NamingTokens", meta = (AutoCreateRefTerm = "InFilter,InContexts"))
+	FNamingTokenResultData EvaluateTokenString(const FString& InTokenString, const FNamingTokenFilterArgs& InFilter, const TArray<UObject*>& InContexts);
+	FNamingTokenResultData EvaluateTokenString(const FString& InTokenString, const FNamingTokenFilterArgs& InFilter = FNamingTokenFilterArgs());
+	
 	/**
 	 * Register tokens as a global namespace. This prevents the need to include the namespace in a token string.
 	 * @param InNamespace The namespace of the tokens to register.
diff --git a/Engine/Plugins/Developer/NamingTokens/Source/NamingTokens/Public/NamingTokensEvaluationData.h b/Engine/Plugins/Developer/NamingTokens/Source/NamingTokens/Public/NamingTokensEvaluationData.h
--- a/Engine/Plugins/Developer/NamingTokens/Source/NamingTokens/Public/NamingTokensEvaluationData.h
+++ b/Engine/Plugins/Developer/NamingTokens/Source/NamingTokens/Public/NamingTokensEvaluationData.h
@@ -21,4 +21,8 @@ public:
 	/** The current date time. */
 	UPROPERTY(BlueprintReadOnly, Category = "NamingTokens")
 	FDateTime CurrentDateTime;
+
+	/** An optional array of context objects for this evaluation. */
+	UPROPERTY(BlueprintReadOnly, Category = "NamingTokens")
+	TArray<TObjectPtr<UObject>> Contexts;
 };
\ No newline at end of file
diff --git a/Engine/Plugins/Developer/NamingTokens/Source/NamingTokens/Public/NamingTokensSpecifiers.h b/Engine/Plugins/Developer/NamingTokens/Source/NamingTokens/Public/NamingTokensSpecifiers.h
new file mode 100644
--- /dev/null
+++ b/Engine/Plugins/Developer/NamingTokens/Source/NamingTokens/Public/NamingTokensSpecifiers.h
@@ -0,0 +1,11 @@
+// Copyright Epic Games, Inc. All Rights Reserved.
+
+#pragma once
+
+#include "CoreMinimal.h"
+
+namespace UE::NamingTokens::Specifiers
+{
+	/** Specifier indicating naming tokens should be used. */
+	const FString UseNamingTokens = TEXT("NamingTokens");
+}
\ No newline at end of file
diff --git a/Engine/Plugins/VirtualProduction/Takes/Source/TakeRecorder/Private/Widgets/SLevelSequenceTakeEditor.cpp b/Engine/Plugins/VirtualProduction/Takes/Source/TakeRecorder/Private/Widgets/SLevelSequenceTakeEditor.cpp
--- a/Engine/Plugins/VirtualProduction/Takes/Source/TakeRecorder/Private/Widgets/SLevelSequenceTakeEditor.cpp
+++ b/Engine/Plugins/VirtualProduction/Takes/Source/TakeRecorder/Private/Widgets/SLevelSequenceTakeEditor.cpp
@@ -116,6 +116,9 @@ void SLevelSequenceTakeEditor::Construct(const FArguments& InArgs)
 	bRequestDetailsRefresh = true;
 	LevelSequenceAttribute = InArgs._LevelSequence;
 
+	OnDetailsPropertiesChangedEvent = InArgs._OnDetailsPropertiesChanged;
+	OnDetailsViewAddedEvent = InArgs._OnDetailsViewAdded;
+
 	DetailsBox = SNew(SScrollBox);
 	DetailsBox->SetScrollBarRightClickDragAllowed(true);
 
@@ -290,6 +293,11 @@ void SLevelSequenceTakeEditor::OnSourcesSelectionChanged(TSharedPtr<ITakeRecorde
 	bRequestDetailsRefresh = true;
 }
 
+void SLevelSequenceTakeEditor::OnDetailsPropertiesChanged(const FPropertyChangedEvent& InEvent)
+{
+	OnDetailsPropertiesChangedEvent.ExecuteIfBound(InEvent);
+}
+
 class FRecordedPropertyCustomization : public IPropertyTypeCustomization
 {
 	static const FString PropertyPathDelimiter;
@@ -918,6 +926,9 @@ void SLevelSequenceTakeEditor::AddDetails(const TPair<const UClass*, TArray<UObj
 
 		Details->SetEnabled(LevelSequenceAttribute.IsSet() && LevelSequenceAttribute.Get()->FindMetaData<UTakeMetaData>() ? !LevelSequenceAttribute.Get()->FindMetaData<UTakeMetaData>()->Recorded() : true);
 
+		Details->OnFinishedChangingProperties().AddSP(this, &SLevelSequenceTakeEditor::OnDetailsPropertiesChanged);
+		OnDetailsViewAddedEvent.ExecuteIfBound(Details);
+		
 		DetailsBox->AddSlot()
 			[
 				Details
diff --git a/Engine/Plugins/VirtualProduction/Takes/Source/TakeRecorder/Private/Widgets/SLevelSequenceTakeEditor.h b/Engine/Plugins/VirtualProduction/Takes/Source/TakeRecorder/Private/Widgets/SLevelSequenceTakeEditor.h
--- a/Engine/Plugins/VirtualProduction/Takes/Source/TakeRecorder/Private/Widgets/SLevelSequenceTakeEditor.h
+++ b/Engine/Plugins/VirtualProduction/Takes/Source/TakeRecorder/Private/Widgets/SLevelSequenceTakeEditor.h
@@ -6,6 +6,7 @@
 #include "Widgets/SCompoundWidget.h"
 #include "UObject/GCObject.h"
 #include "UObject/ObjectKey.h"
+#include "PropertyEditorDelegates.h"
 
 struct FAssetData;
 struct ITakeRecorderSourceTreeItem;
@@ -24,12 +25,16 @@ template<class> class TSubclassOf;
 class SLevelSequenceTakeEditor : public SCompoundWidget
 {
 public:
+	DECLARE_DELEGATE_OneParam(FOnDetailsPropertiesChanged, const FPropertyChangedEvent&);
+	DECLARE_DELEGATE_OneParam(FOnDetailsViewAdded, const TWeakPtr<IDetailsView>&);
 
 	SLATE_BEGIN_ARGS(SLevelSequenceTakeEditor)
 		: _LevelSequence(nullptr)
 		{}
 
 		SLATE_ATTRIBUTE(ULevelSequence*, LevelSequence)
+		SLATE_EVENT(FOnDetailsPropertiesChanged, OnDetailsPropertiesChanged)
+		SLATE_EVENT(FOnDetailsViewAdded, OnDetailsViewAdded)
 
 	SLATE_END_ARGS()
 
@@ -83,6 +88,10 @@ private:
 	bool CanAddSourceFromClass(TSubclassOf<UTakeRecorderSource> SourceClass);
 	
 	void OnSourcesSelectionChanged(TSharedPtr<ITakeRecorderSourceTreeItem>, ESelectInfo::Type);
+	/**
+	 * When a details box has a property changed.
+	 */
+	void OnDetailsPropertiesChanged(const FPropertyChangedEvent& InEvent);
 
 private:
 
@@ -95,4 +104,9 @@ private:
 	TMap<FObjectKey, TSharedPtr<IDetailsView>> ClassToDetailsView;
 
 	TArray<TWeakObjectPtr<>> ExternalSettingsObjects;
+
+	/** Called when properties in the details panel have changed. */
+	FOnDetailsPropertiesChanged OnDetailsPropertiesChangedEvent;
+	/** Called when a details view has been added. */
+	FOnDetailsViewAdded OnDetailsViewAddedEvent;
 };
\ No newline at end of file
diff --git a/Engine/Plugins/VirtualProduction/Takes/Source/TakeRecorder/Private/Widgets/STakeRecorderCockpit.cpp b/Engine/Plugins/VirtualProduction/Takes/Source/TakeRecorder/Private/Widgets/STakeRecorderCockpit.cpp
--- a/Engine/Plugins/VirtualProduction/Takes/Source/TakeRecorder/Private/Widgets/STakeRecorderCockpit.cpp
+++ b/Engine/Plugins/VirtualProduction/Takes/Source/TakeRecorder/Private/Widgets/STakeRecorderCockpit.cpp
@@ -2,6 +2,7 @@
 
 #include "Widgets/STakeRecorderCockpit.h"
 #include "Widgets/SOverlay.h"
+#include "Widgets/Layout/SScrollBox.h"
 #include "Widgets/TakeRecorderWidgetConstants.h"
 #include "Widgets/STakeRecorderTabContent.h"
 #include "TakesCoreBlueprintLibrary.h"
@@ -39,7 +40,6 @@
 #include "Widgets/Layout/SScrollBox.h"
 #include "Widgets/Text/STextBlock.h"
 #include "Widgets/Input/SComboButton.h"
-#include "Widgets/Input/SMultiLineEditableTextBox.h"
 #include "Widgets/Input/SEditableTextBox.h"
 #include "Widgets/Input/SSpinBox.h"
 #include "Widgets/Input/SButton.h"
@@ -54,6 +54,7 @@
 #include "EditorFontGlyphs.h"
 
 // UnrealEd includes
+#include "Algo/Compare.h"
 #include "ScopedTransaction.h"
 #include "Editor/EditorEngine.h"
 #include "Editor.h"
@@ -63,10 +64,18 @@
 #include "Dialog/SMessageDialog.h"
 #include "Kismet2/DebuggerCommands.h"
 #include "Subsystems/AssetEditorSubsystem.h"
+#include "PropertyEditorModule.h"
+#include "UObject/TextProperty.h"
 
 #include "ISequencer.h"
 #include "ILevelSequenceEditorToolkit.h"
 
+// NamingTokens includes
+#include "ITakeRecorderNamingTokensModule.h"
+#include "NamingTokens.h"
+#include "NamingTokensEditorSubsystem.h"
+#include "NamingTokensSpecifiers.h"
+
 extern UNREALED_API UEditorEngine* GEditor;
 
 #define LOCTEXT_NAMESPACE "STakeRecorderCockpit"
@@ -74,6 +83,7 @@
 STakeRecorderCockpit::~STakeRecorderCockpit()
 {
 	UTakeRecorder::OnRecordingInitialized().Remove(OnRecordingInitializedHandle);
+	FEditorDelegates::PostUndoRedo.Remove((PostUndoRedoDelegateHandle));
 
 	if (FAssetRegistryModule* AssetRegistryModule = FModuleManager::GetModulePtr<FAssetRegistryModule>("AssetRegistry"))
 	{
@@ -88,6 +98,8 @@
 	{
 		GEditor->CancelTransaction(TransactionIndex);
 	}
+
+	UnbindNamingTokensEvents();
 }
 
 void STakeRecorderCockpit::AddReferencedObjects(FReferenceCollector& Collector)
@@ -131,7 +143,20 @@
 
 	virtual bool IsCharacterValid(TCHAR InChar) const override { return true; }
 };
-	
+
+namespace UE::TakeRecorder::Private
+{
+	/** The pre-defined token name representing TAKE. */
+	const FString TokenKeyTake(TEXT("take"));
+	/** The pre-defined token name representing SLATE. */
+	const FString TokenKeySlate(TEXT("slate"));
+
+	/** Property name of the TakeSaveDir. */
+	const FName TakeSaveDirPropertyName(GET_MEMBER_NAME_CHECKED(FTakeRecorderProjectParameters, TakeSaveDir));
+	/** Property name of the RootTakeSaveDir. */
+	const FName RootTakeSaveDirPropertyName(GET_MEMBER_NAME_CHECKED(FTakeRecorderProjectParameters, RootTakeSaveDir));
+}
+
 void STakeRecorderCockpit::Construct(const FArguments& InArgs)
 {
 	FAssetRegistryModule& AssetRegistryModule = FModuleManager::LoadModuleChecked<FAssetRegistryModule>(TEXT("AssetRegistry"));
@@ -143,6 +168,10 @@
 
 	TakeMetaData = nullptr;
 	TransientTakeMetaData = nullptr;
+	NamingTokensData = static_cast<TStrongObjectPtr<UTakeRecorderNamingTokensWidgetData>>(NewObject<UTakeRecorderNamingTokensWidgetData>());
+	NamingTokensData->SetFlags(RF_Transactional);
 
+	PostUndoRedoDelegateHandle = FEditorDelegates::PostUndoRedo.AddRaw(this, &STakeRecorderCockpit::PostUndoRedo);
+
 	CachedTakeSlate.Empty();
 	CachedTakeNumber = -1;
@@ -168,6 +197,7 @@
 	DigitsTypeInterface = MakeShareable(new FDigitsTypeInterface);
 
 	BindCommands();
+	BindToNamingTokenEvents();
 
 	TransactionIndex = INDEX_NONE;
 
@@ -184,7 +214,7 @@
 
 		SNew(SVerticalBox)
 
-		// Slate, Take #, and Record Button 
+		// Slate, Take #, User Defined tokens, and Record Button 
 		+SVerticalBox::Slot()
 		.AutoHeight()
 		[
@@ -194,95 +224,17 @@
 				SNew(SHorizontalBox)
 
 				+ SHorizontalBox::Slot()
-				.FillWidth(0.6)
 				[
-					SNew(SOverlay)
-
-					+ SOverlay::Slot()
-					.VAlign(VAlign_Top)
-					.HAlign(HAlign_Left)
-					.Padding(2.0f, 2.0f)
+					SNew(SScrollBox).Orientation(Orient_Horizontal)
+					+ SScrollBox::Slot()
 					[
-						SNew(STextBlock)
-						.TextStyle(FTakeRecorderStyle::Get(), "TakeRecorder.TextBox")
-						.Text(LOCTEXT("SlateLabel", "SLATE"))
-					]
-
-					+ SOverlay::Slot()
-					.VAlign(VAlign_Center)
-					.HAlign(HAlign_Center)
-					[
-						SNew(SEditableTextBox)
-						.IsEnabled(this, &STakeRecorderCockpit::EditingMetaData)
-						.Style(FTakeRecorderStyle::Get(), "TakeRecorder.EditableTextBox")
-						.Font(FTakeRecorderStyle::Get().GetFontStyle("TakeRecorder.Cockpit.LargeText"))
-						.HintText(LOCTEXT("EnterSlate_Hint", "<slate>"))
-						.Justification(ETextJustify::Center)
-						.SelectAllTextWhenFocused(true)
-						.Text(this, &STakeRecorderCockpit::GetSlateText)
-						.OnTextCommitted(this, &STakeRecorderCockpit::SetSlateText)
-					]
-				]
-
-				+ SHorizontalBox::Slot()
-				.FillWidth(0.4)
-				[
-					SNew(SOverlay)
-
-					+ SOverlay::Slot()
-					.VAlign(VAlign_Top)
-					.HAlign(HAlign_Left)
-					.Padding(2.0f, 2.0f)
-					[
 						SNew(SHorizontalBox)
-
+						// Detected user tokens
 						+ SHorizontalBox::Slot()
 						.AutoWidth()
 						[
-							SNew(STextBlock)
-							.TextStyle(FTakeRecorderStyle::Get(), "TakeRecorder.TextBox")
-							.Text(LOCTEXT("TakeLabel", "TAKE"))
+							SAssignNew(UserTokensBox, SHorizontalBox)
 						]
-
-						+ SHorizontalBox::Slot()
-						.Padding(2.f, 0.f)
-						.VAlign(VAlign_Center)
-						.AutoWidth()
-						[
-							SNew(SButton)
-							.ButtonStyle(FAppStyle::Get(), "NoBorder")
-							.OnClicked(this, &STakeRecorderCockpit::OnSetNextTakeNumber)
-							.ForegroundColor(FSlateColor::UseForeground())
-							.Visibility(this, &STakeRecorderCockpit::GetTakeWarningVisibility)
-							.Content()
-							[
-								SNew(STextBlock)
-								.ToolTipText(this, &STakeRecorderCockpit::GetTakeWarningText)
-								.Font(FAppStyle::Get().GetFontStyle("FontAwesome.8"))
-								.Text(FEditorFontGlyphs::Exclamation_Triangle)
-							]
-						]
-					]
-
-					+ SOverlay::Slot()
-					.VAlign(VAlign_Center)
-					.HAlign(HAlign_Center)
-					[
-						SNew(SSpinBox<int32>)
-						.IsEnabled(this, &STakeRecorderCockpit::EditingMetaData)
-						.ContentPadding(FMargin(8.f, 0.f))
-						.Style(FTakeRecorderStyle::Get(), "TakeRecorder.TakeInput")
-						.Font(FTakeRecorderStyle::Get().GetFontStyle("TakeRecorder.Cockpit.GiantText"))
-						.Justification(ETextJustify::Center)
-						.Value(this, &STakeRecorderCockpit::GetTakeNumber)
-						.Delta(1)
-						.MinValue(1)
-						.MaxValue(TOptional<int32>())
-						.OnBeginSliderMovement(this, &STakeRecorderCockpit::OnBeginSetTakeNumber)
-						.OnValueChanged(this, &STakeRecorderCockpit::SetTakeNumber)
-						.OnValueCommitted(this, &STakeRecorderCockpit::SetTakeNumber_FromCommit)
-						.OnEndSliderMovement(this, &STakeRecorderCockpit::OnEndSetTakeNumber)
-						.TypeInterface(DigitsTypeInterface)
 					]
 				]
 
@@ -375,6 +327,21 @@
 			]
 		]
 
+		// Take Save Dir Example text
+		+SVerticalBox::Slot()
+		.AutoHeight()
+		[
+			SNew(SBorder)
+			.BorderImage_Lambda([this]{ return Reviewing() ? FTakeRecorderStyle::Get().GetBrush("TakeRecorder.TakeRecorderReviewBorder") : FAppStyle::GetBrush("ToolPanel.DarkGroupBorder"); })
+			[
+				SNew(SEditableTextBox)
+				.IsReadOnly(true)
+				.Style(FTakeRecorderStyle::Get(), "TakeRecorder.EditableTextBox")
+				.Font(FTakeRecorderStyle::Get().GetFontStyle("TakeRecorder.Cockpit.MediumText"))
+				.Text(this, &STakeRecorderCockpit::GetEvaluatedTakeSaveDirText)
+			]
+		]
+
 		// Timestamp, Duration, Description and Remaining Metadata
 		+SVerticalBox::Slot()
 		.AutoHeight()
@@ -525,6 +492,8 @@
 				Widget
 			];
 	}
+	
+	bRefreshUndefinedTokens = true;
 }
 
 bool STakeRecorderCockpit::CanStartRecording(FText& OutErrorText) const
@@ -723,6 +692,8 @@
 			{
 				TransientTakeMetaData->SetTakeNumber(NextTakeNumber, false);
 			}
+
+			UTakeMetaData::SetMostRecentMetaData(TransientTakeMetaData);
 		}
 
 		NewMetaDataThisTick = TransientTakeMetaData;
@@ -758,6 +729,16 @@
 void STakeRecorderCockpit::Tick(const FGeometry& AllottedGeometry, const double InCurrentTime, const float InDeltaTime)
 {
 	Refresh();
+	if (bRefreshUndefinedTokens)
+	{
+		bRefreshUndefinedTokens = false;
+		RefreshUndefinedTokens();
+	}
 }
 
+void STakeRecorderCockpit::PostUndoRedo()
+{
+	RefreshUndefinedTokens();
+}
+
 void STakeRecorderCockpit::Refresh()
@@ -766,6 +747,62 @@
 	UpdateRecordError();
 }
 
+void STakeRecorderCockpit::NotifyPropertyUpdated(const FPropertyChangedEvent& InPropertyChangedEvent)
+{
+	if (InPropertyChangedEvent.Property)
+	{
+		FString ResultStr;
+
+		for (int32 Idx = 0; Idx < InPropertyChangedEvent.GetNumObjectsBeingEdited(); ++Idx )
+		{
+			if (const UObject* Object = InPropertyChangedEvent.GetObjectBeingEdited(Idx))
+			{
+				bool bForce = false;
+				const void* Container = Object;
+				TMap<FString, int32> ArrayIndicesPerObject;
+				if (InPropertyChangedEvent.MemberProperty && InPropertyChangedEvent.MemberProperty != InPropertyChangedEvent.Property)
+				{
+					// This could be a member of a struct, which has its own metadata and container.
+					if (InPropertyChangedEvent.MemberProperty->HasMetaData(*UE::NamingTokens::Specifiers::UseNamingTokens)
+						||
+						// @hack: When dealing with the RootTakeSaveDir, we don't have a good way of getting the meta specifier, so we need to
+						// check if the property is contained in the ArrayIndices. The MemberProperty *should* be the FDirectoryPath property,
+						// but it's actually the FTakeRecorderProjectParameters property.
+						// ChainProperty might be better for this, but we're limited to what the details panel property changed event provides us.
+						// @todo NamingTokens - If we move to a more global system, we need a more robust way of checking for meta specifiers on FDirectoryPaths.
+						// @note: The PresetSaveLocation we use doesn't have this issue and the MemberProperty is correct.
+						(InPropertyChangedEvent.GetArrayIndicesPerObject(0, ArrayIndicesPerObject)
+							&& ArrayIndicesPerObject.Contains(UE::TakeRecorder::Private::RootTakeSaveDirPropertyName.ToString())))
+					{
+						bForce = true;
+					}
+					
+					Container = InPropertyChangedEvent.MemberProperty->ContainerPtrToValuePtr<void>(Object);
+				}
+				
+				EvaluateTokensFromProperty(InPropertyChangedEvent.Property, Container, bForce);
+				CreateUserTokensUI();
+			}
+		}
+	}
+}
+
+void STakeRecorderCockpit::NotifyDetailsViewAdded(const TWeakPtr<IDetailsView>& InDetailsView)
+{
+	// Cleanup stale entries. This list can change when the user clicks on or off an actor.
+	DetailViews.RemoveAll([](const TWeakPtr<IDetailsView>& DetailsView)
+	{
+		return !DetailsView.IsValid();
+	});
+	
+	DetailViews.Add(InDetailsView);
+
+	// We signal a refresh on the next tick rather than now so all details views have time to process. Most likely
+	// this method is being called multiple times since there are multiple objects being added. We want to refresh after
+	// all have been added so we don't remove any custom tokens that are still defined.
+	bRefreshUndefinedTokens = true;
+}
+
 FText STakeRecorderCockpit::GetSlateText() const
 {
 	return FText::FromString(TakeMetaData->GetSlate());
@@ -807,6 +844,11 @@
 	}
 }
 
+FText STakeRecorderCockpit::GetEvaluatedTakeSaveDirText() const
+{
+	return NamingTokensData->EvaluatedTextValue;
+}
+
 void STakeRecorderCockpit::SetFrameRate(FFrameRate InFrameRate, bool bFromTimecode)
 {
 	if (TakeMetaData)
@@ -871,6 +913,8 @@
 		{
 			TakeMetaData->SetTakeNumber(NextTakeNumber);
 		}
+
+		OnTokenValueUpdated();
 	}
 }
 
@@ -927,6 +971,8 @@
 		TakeMetaData->SetTakeNumber(InNewTakeNumber, false);
 		bAutoApplyTakeNumber = false;
 	}
+	
+	OnTokenValueUpdated();
 }
 
 void STakeRecorderCockpit::SetTakeNumber_FromCommit(int32 InNewTakeNumber, ETextCommit::Type InCommitType)
@@ -947,6 +993,7 @@
 	}
 
 	bAutoApplyTakeNumber = false;
+	OnTokenValueUpdated();
 }
 
 void STakeRecorderCockpit::OnEndSetTakeNumber(int32 InFinalValue)
@@ -1317,4 +1364,526 @@
 	return MenuBuilder.MakeWidget();
 }
 
+void STakeRecorderCockpit::BindToNamingTokenEvents()
+{
+	if (GIsEditor && !NamingTokensData->TakeRecorderNamingTokens.IsValid())
+	{
+		NamingTokensData->TakeRecorderNamingTokens = GEditor->GetEditorSubsystem<UNamingTokensEditorSubsystem>()->GetNamingTokens(ITakeRecorderNamingTokensModule::GetTakeRecorderNamespace());
+		if (ensure(NamingTokensData->TakeRecorderNamingTokens.IsValid()))
+		{
+			const TStrongObjectPtr<UNamingTokens> TakeRecorderNamingTokensPin = NamingTokensData->TakeRecorderNamingTokens.Pin();
+			TakeRecorderNamingTokensPin->RegisterExternalTokens(NamingTokensData->NamingTokensExternalGuid);
+			TakeRecorderNamingTokensPin->GetOnPreEvaluateEvent().AddSP(this, &STakeRecorderCockpit::OnTakeRecorderNamingTokensPreEvaluate);
+		}
+	}
+}
+
+void STakeRecorderCockpit::UnbindNamingTokensEvents()
+{
+	if (const TStrongObjectPtr<UNamingTokens> TakeRecorderNamingTokensPin = NamingTokensData->TakeRecorderNamingTokens.Pin())
+	{
+		TakeRecorderNamingTokensPin->UnregisterExternalTokens(NamingTokensData->NamingTokensExternalGuid);
+		TakeRecorderNamingTokensPin->GetOnPreEvaluateEvent().RemoveAll(this);
+	}
+}
+
+FText STakeRecorderCockpit::GetCustomTokenTextValue(FString InTokenKey) const
+{
+	if (const FText* Value = NamingTokensData->UserDefinedTokens.Find(InTokenKey))
+	{
+		return *Value;
+	}
+
+	return FText::GetEmpty();
+}
+
+void STakeRecorderCockpit::SetCustomTokenTextValue(const FText& InNewText, ETextCommit::Type InCommitType, FString InTokenKey)
+{
+	FScopedTransaction Transaction(LOCTEXT("SetCustomTokenTextValue_Transaction", "Set Custom Token Text"));
+	NamingTokensData->Modify();
+	NamingTokensData->UserDefinedTokens.FindOrAdd(InTokenKey) = InNewText;
+	NamingTokensData->VisibleUserTokens.Add(InTokenKey);
+	OnTokenValueUpdated();
+}
+
+void STakeRecorderCockpit::OnTokenValueUpdated()
+{
+	EvaluateTakeSaveDirTokens();
+}
+
+void STakeRecorderCockpit::EvaluateTakeSaveDirTokens()
+{
+	if (CachedTakeSaveDirProperty.IsValid() && CachedTakeSaveDirContainer)
+	{
+		EvaluateTokensFromProperty(CachedTakeSaveDirProperty.Get(), CachedTakeSaveDirContainer);
+	}
+}
+
+void STakeRecorderCockpit::OnTakeRecorderNamingTokensPreEvaluate(const FNamingTokensEvaluationData& InEvaluationData)
+{
+	if (NamingTokensData->TakeRecorderNamingTokens.IsValid())
+	{
+		TArray<FNamingTokenData>& ExternalTokens = NamingTokensData->TakeRecorderNamingTokens->GetExternalTokensChecked(NamingTokensData->NamingTokensExternalGuid);
+
+		ExternalTokens.Reset(NamingTokensData->UserDefinedTokens.Num());
+		for (const TTuple<FString, FText>& UserToken : NamingTokensData->UserDefinedTokens)
+		{
+			ExternalTokens.Add({
+				UserToken.Key,
+				FText::FromString(UserToken.Key),
+				FNamingTokenData::FTokenProcessorDelegateNative::CreateLambda([Value = UserToken.Value]() {
+					return Value;
+				})
+			});
+		}
+	}
+}
+
+void STakeRecorderCockpit::RefreshUndefinedTokens()
+{
+	// Clicking on different actors will remove the old details panels -- clear out the fields
+	// we are tracking. They will be repopulated when evaluating the tokens below.
+	NamingTokensData->FieldToUndefinedKeys.Empty();
+	
+	bool bHasMinRequiredObjects = false;
+	const UTakeRecorderUserSettings* UserSettings = GetDefault<UTakeRecorderUserSettings>();
+	const UTakeRecorderProjectSettings* ProjectSettings = GetDefault<UTakeRecorderProjectSettings>();
+
+	for (const TWeakPtr<IDetailsView>& DetailsView : DetailViews)
+	{
+		if (const TSharedPtr<IDetailsView> DetailsPin = DetailsView.Pin())
+		{
+			// Look if the default user or project settings are available. This should only fail if the user has hidden these objects.
+			if (DetailsPin->GetSelectedObjects().Contains(UserSettings)
+				|| DetailsPin->GetSelectedObjects().Contains(ProjectSettings))
+			{
+				bHasMinRequiredObjects = true;
+			}
+		}
+		UpdateUndefinedTokensFromDetailsView(DetailsView);
+	}
+
+	// We always need to display tokens for the user/project settings. This should be done after any current details tokens are calculated
+	// to maintain order.
+	if (!bHasMinRequiredObjects)
+	{
+		UpdateUndefinedTokensFromTemporaryDetailsView();
+	}
+	
+	CreateUserTokensUI();
+}
+
+void STakeRecorderCockpit::UpdateUndefinedTokensFromDetailsView(const TWeakPtr<IDetailsView>& InDetailsView)
+{
+	if (const TSharedPtr<IDetailsView> DetailsViewPin = InDetailsView.Pin())
+	{
+		const TArray<FPropertyPath> Properties = DetailsViewPin->GetPropertiesInOrderDisplayed();
+		const TArray<TWeakObjectPtr<UObject>> Objects = DetailsViewPin->GetSelectedObjects();
+		for (const TWeakObjectPtr<UObject>& EditedObject : Objects)
+		{
+			if (!EditedObject.IsValid())
+			{
+				continue;
+			}
+			auto GetCorrectContainer = [&EditedObject](const FProperty* InProperty) -> UObject*
+			{
+				if (InProperty)
+				{
+					// Locate the correct container to use. When using a project settings object it could also be an additional settings object.
+					// The additional settings isn't a UPROPERTY and are added to the details via TakeRecorder customization, so they won't be discovered
+					// by GetSelectedObjects. Ideally we wouldn't be aware that we are of the type of object we are editing.
+					if (const UTakeRecorderProjectSettings* ProjectSettings = Cast<UTakeRecorderProjectSettings>(EditedObject))
+					{
+						for (const TWeakObjectPtr<UObject>& AdditionalObj : ProjectSettings->AdditionalSettings)
+						{
+							if (UClass* OwnerClass = InProperty->GetOwnerClass())
+							{
+								if (AdditionalObj.IsValid() && AdditionalObj->IsA(OwnerClass))
+								{
+									return AdditionalObj.Get();
+								}
+							}
+						}
+					}
+				}
+				return EditedObject.Get();
+			};
+		
+			for (const FPropertyPath& PropertyPath : Properties)
+			{
+				void* CurrentContainer = GetCorrectContainer(PropertyPath.GetRootProperty().Property.Get());
+				const int32 NumSegments = PropertyPath.GetNumProperties();
+
+				// For multi-length segments adjust the container.
+				for (int32 Idx = 0; Idx < NumSegments - 1; ++Idx)
+				{
+					const FProperty* Prop = PropertyPath.GetPropertyInfo(Idx).Property.Get();
+					if (!Prop)
+					{
+						CurrentContainer = nullptr;
+						break;
+					}
+					CurrentContainer = Prop->ContainerPtrToValuePtr<void>(CurrentContainer);
+					if (!CurrentContainer)
+					{
+						break;
+					}
+				}
+
+				const FProperty* LeafProperty = PropertyPath.GetLeafMostProperty().Property.Get();
+				if (!LeafProperty || !CurrentContainer)
+				{
+					continue;
+				}
+
+				bool bForce = false;
+				if (const FStructProperty* StructProperty = CastField<FStructProperty>(LeafProperty))
+				{
+					if (StructProperty->HasMetaData(*UE::NamingTokens::Specifiers::UseNamingTokens))
+					{
+						// The property we evaluate won't have the metadata so we need to force it.
+						bForce = true;
+					
+						// Need to adjust the container for the struct.
+						CurrentContainer = StructProperty->ContainerPtrToValuePtr<void>(CurrentContainer);
+				
+						// Directory and FilePaths won't have their string Path iterated here, so find it manually. Not USTRUCTs so have to check name.
+						const FName StructName = StructProperty->Struct->GetFName();
+						if (StructName == FName("DirectoryPath"))
+						{
+							LeafProperty = FindFProperty<FStrProperty>(StructProperty->Struct, GET_MEMBER_NAME_CHECKED(FDirectoryPath, Path));
+						}
+						else if (StructName == FName("FilePath"))
+						{
+							LeafProperty = FindFProperty<FStrProperty>(StructProperty->Struct, GET_MEMBER_NAME_CHECKED(FFilePath, FilePath));
+						}
+					}
+				}
+
+				EvaluateTokensFromProperty(LeafProperty, CurrentContainer, bForce);
+			}
+		}
+	}
+}
+
+void STakeRecorderCockpit::UpdateUndefinedTokensFromTemporaryDetailsView()
+{
+	if (TemporaryDetailsViews.Num() == 0)
+	{
+		UTakeRecorderProjectSettings* ProjectSettings = GetMutableDefault<UTakeRecorderProjectSettings>();
+		UTakeRecorderUserSettings* UserSettings = GetMutableDefault<UTakeRecorderUserSettings>();
+		
+		const TArray<UObject*> ObjectsToDisplay { ProjectSettings, UserSettings };
+
+		for (UObject* Object : ObjectsToDisplay)
+		{
+			FPropertyEditorModule& PropertyEditorModule = FModuleManager::Get().LoadModuleChecked<FPropertyEditorModule>("PropertyEditor");
+			FDetailsViewArgs DetailsViewArgs;
+			DetailsViewArgs.bAllowSearch = false;
+			DetailsViewArgs.NameAreaSettings = FDetailsViewArgs::HideNameArea;
+			DetailsViewArgs.bHideSelectionTip = true;
+			DetailsViewArgs.bShowScrollBar = false;
+
+			TSharedPtr<IDetailsView> TemporaryDetailsView = PropertyEditorModule.CreateDetailView(DetailsViewArgs);
+
+			// Have to pass individual objects through, won't work if we just use one view with all our objects.
+			const TArray<UObject*> Objects { Object };
+			TemporaryDetailsView->SetObjects(Objects);
+			TemporaryDetailsView->SetEnabled(false);
+
+			TemporaryDetailsViews.Add(MoveTemp(TemporaryDetailsView));
+		}
+	}
+
+	for (const TSharedPtr<IDetailsView>& TemporaryDetailsView : TemporaryDetailsViews)
+	{
+		UpdateUndefinedTokensFromDetailsView(TemporaryDetailsView.ToWeakPtr());
+	}
+}
+
+void STakeRecorderCockpit::EvaluateTokensFromProperty(const FProperty* InProperty, const void* InContainer, bool bForce)
+{
+	check(InProperty);
+	check(InContainer);
+
+	const FProperty* OwnerProperty = InProperty->GetOwnerProperty();
+	const bool bIsNamingTokensField = bForce
+		|| InProperty->HasMetaData(*UE::NamingTokens::Specifiers::UseNamingTokens)
+		|| (OwnerProperty && OwnerProperty->HasMetaData(*UE::NamingTokens::Specifiers::UseNamingTokens));
+	
+	if (!bIsNamingTokensField)
+	{
+		return;
+	}
+	
+	FText ResultText;
+	if (const FStrProperty* StringProp = CastField<FStrProperty>(InProperty))
+	{
+		if (const FString* ValuePtr = StringProp->ContainerPtrToValuePtr<FString>(InContainer))
+		{
+			ResultText = FText::FromString(*ValuePtr);
+		}
+	}
+	else if (const FTextProperty* TextProp = CastField<FTextProperty>(InProperty))
+	{
+		if (const FText* TextPtr = TextProp->ContainerPtrToValuePtr<FText>(InContainer))
+		{
+			ResultText = *TextPtr;
+		}
+	}
+	
+	FNamingTokenResultData NamingTokenResultData;
+	if (!ResultText.IsEmpty())
+	{
+		// Create an identifier which handles duplicate property names on different objects.
+		// Only valid for the container's life cycle, which should match this widget.
+		const FString FieldName = FString::Printf(TEXT("%s_%p"), *InProperty->GetName(), InContainer);
+
+		// We will automatically inject our tokens to the take recorder naming tokens via our pre-evaluate hook.
+		NamingTokenResultData = TakeMetaData->ProcessTokens(ResultText, nullptr);
+		
+		// Manually handle after all evaluation has finished. We don't need to listen for a post-evaluate event
+		// because we want the finalized text processed for all tokens classes.
+		// Build list of unidentified tokens, in-order of appearance. Include external tokens since they are being added by us.
+		TArray<FString> UndefinedTokenKeys;
+		for (const FNamingTokenValueData& Token : NamingTokenResultData.TokenValues)
+		{
+			if ((!Token.bWasEvaluated
+				|| NamingTokensData->UserDefinedTokens.Contains(Token.TokenKey))
+					// Include pre-defined take and slate keys so they will get sorted with user tokens.
+				|| (Token.TokenKey == UE::TakeRecorder::Private::TokenKeyTake
+					|| Token.TokenKey == UE::TakeRecorder::Private::TokenKeySlate))
+			{
+				UndefinedTokenKeys.Add(Token.TokenKey);
+			}
+		}
+
+		if (!FieldName.IsEmpty())
+		{
+			// Track the unidentified keys for this field.
+			TArray<FString>& TokenKeysForField = NamingTokensData->FindOrAddTokenKeysForField(FieldName);
+			TokenKeysForField = UndefinedTokenKeys;
+		}
+	}
+
+	// Cache the value if we're the TakeSaveDir, so it can be displayed in the UI as an example.
+	if (InProperty->GetFName() == UE::TakeRecorder::Private::TakeSaveDirPropertyName)
+	{
+		CachedTakeSaveDirProperty = InProperty;
+		CachedTakeSaveDirContainer = InContainer;
+		NamingTokensData->Modify();
+		NamingTokensData->EvaluatedTextValue = NamingTokenResultData.EvaluatedText;
+	}
+}
+
+void STakeRecorderCockpit::CreateUserTokensUI()
+{
+	// Assemble the custom tokens in use. This step needs to happen after NamingTokensData->FieldToUndefinedKeys has
+	// been fully populated so we don't remove custom tokens that are still in use.
+	{
+		const TMap<FString, FText> CustomUserTokensCopy = NamingTokensData->UserDefinedTokens;
+		for (const TTuple<FString, FText>& CustomToken : CustomUserTokensCopy)
+		{
+			// If no field holds a reference to this key then we can remove it.
+			if (!NamingTokensData->IsTokenKeyUndefined(CustomToken.Key))
+			{
+				// Only remove the visible entry, but leave the value in CustomUserTokens. If the user is switching
+				// between objects it's possible one object has fields with different tokens than another, and we don't
+				// want to lose the old entered data.
+				NamingTokensData->VisibleUserTokens.Remove(CustomToken.Key);
+			}
+		}
+
+		// Look for new unidentified tokens that a user can define.
+		for (const FTakeRecorderNamingTokensFieldMapping& UndefinedTokens : NamingTokensData->FieldToUndefinedKeys)
+		{
+			for (const FString& UndefinedToken : UndefinedTokens.UndefinedKeys)
+			{
+				if (!NamingTokensData->UserDefinedTokens.Contains(UndefinedToken))
+				{
+					// Check for existence so we don't overwrite the user defined value.
+					NamingTokensData->UserDefinedTokens.Add(UndefinedToken);
+				}
+				NamingTokensData->VisibleUserTokens.Add(UndefinedToken);
+			}
+		}
+	}
+
+	// Make sure our take dir token example is up-to-date.
+	EvaluateTakeSaveDirTokens();
+
+	// @note User defined tokens are case-insensitive right now. We need to add FLocKeySetFuncs to make it case-sensitive, but because
+	// we only display the token key in UPPERCASE for take recorder, we shouldn't allow the same key with different cases.
+	TSet<FString> UsedTokens;
+	UserTokensBox->ClearChildren();
+	
+	// Display tokens in order from their field appearance, to the token location in the string, but do not duplicate.
+	for (const FTakeRecorderNamingTokensFieldMapping& FieldMapping : NamingTokensData->FieldToUndefinedKeys)
+	{
+		for (const FString& TokenString : FieldMapping.UndefinedKeys)
+		{
+			if (UsedTokens.Contains(TokenString) || !NamingTokensData->VisibleUserTokens.Contains(TokenString))
+			{
+				// Skip if this was already reference in this field or a previous one.
+				continue;
+			}
+			ensure(!TokenString.IsEmpty());
+
+			UsedTokens.Add(TokenString);
+
+			constexpr float MinWidth = 100.f;
+			constexpr float MaxWidth = 175.f;
+
+			// Take and Slate have special handling but still need to show up in the order the user referenced them.
+			if (TokenString == UE::TakeRecorder::Private::TokenKeySlate)
+			{
+				UserTokensBox->AddSlot()
+				.HAlign(HAlign_Fill)
+				.MinWidth(MinWidth)
+				.MaxWidth(MaxWidth)
+				.AutoWidth()
+				[
+					SNew(SOverlay)
+
+					+ SOverlay::Slot()
+					.VAlign(VAlign_Top)
+					.HAlign(HAlign_Left)
+					.Padding(2.0f, 2.0f)
+					[
+						SNew(STextBlock)
+						.TextStyle(FTakeRecorderStyle::Get(), "TakeRecorder.TextBox")
+						.Text(LOCTEXT("SlateLabel", "SLATE"))
+					]
+
+					+ SOverlay::Slot()
+					.VAlign(VAlign_Center)
+					.HAlign(HAlign_Center)
+					[
+						SNew(SEditableTextBox)
+						.IsEnabled(this, &STakeRecorderCockpit::EditingMetaData)
+						.Style(FTakeRecorderStyle::Get(), "TakeRecorder.EditableTextBox")
+						.Font(FTakeRecorderStyle::Get().GetFontStyle("TakeRecorder.Cockpit.LargeText"))
+						.HintText(LOCTEXT("EnterSlate_Hint", "<slate>"))
+						.Justification(ETextJustify::Left)
+						.SelectAllTextWhenFocused(true)
+						.Text(this, &STakeRecorderCockpit::GetSlateText)
+						.OnTextCommitted(this, &STakeRecorderCockpit::SetSlateText)
+					]
+				];
+			}
+			else if (TokenString == UE::TakeRecorder::Private::TokenKeyTake)
+			{
+				UserTokensBox->AddSlot()
+				.HAlign(HAlign_Fill)
+				.MinWidth(MinWidth)
+				.MaxWidth(MaxWidth)
+				.AutoWidth()
+				[
+					SNew(SOverlay)
+					+ SOverlay::Slot()
+					.VAlign(VAlign_Top)
+					.HAlign(HAlign_Left)
+					.Padding(2.0f, 2.0f)
+					[
+						SNew(SHorizontalBox)
+
+						+ SHorizontalBox::Slot()
+						.AutoWidth()
+						[
+							SNew(STextBlock)
+							.TextStyle(FTakeRecorderStyle::Get(), "TakeRecorder.TextBox")
+							.Text(LOCTEXT("TakeLabel", "TAKE"))
+						]
+
+						+ SHorizontalBox::Slot()
+						.Padding(2.f, 0.f)
+						.VAlign(VAlign_Center)
+						.AutoWidth()
+						[
+							SNew(SButton)
+							.ButtonStyle(FAppStyle::Get(), "NoBorder")
+							.OnClicked(this, &STakeRecorderCockpit::OnSetNextTakeNumber)
+							.ForegroundColor(FSlateColor::UseForeground())
+							.Visibility(this, &STakeRecorderCockpit::GetTakeWarningVisibility)
+							.Content()
+							[
+								SNew(STextBlock)
+								.ToolTipText(this, &STakeRecorderCockpit::GetTakeWarningText)
+								.Font(FAppStyle::Get().GetFontStyle("FontAwesome.8"))
+								.Text(FEditorFontGlyphs::Exclamation_Triangle)
+							]
+						]
+					]
+
+					+ SOverlay::Slot()
+					.VAlign(VAlign_Center)
+					.HAlign(HAlign_Center)
+					[
+						SNew(SSpinBox<int32>)
+						.IsEnabled(this, &STakeRecorderCockpit::EditingMetaData)
+						.ContentPadding(FMargin(8.f, 0.f))
+						.Style(FTakeRecorderStyle::Get(), "TakeRecorder.TakeInput")
+						.Font(FTakeRecorderStyle::Get().GetFontStyle("TakeRecorder.Cockpit.GiantText"))
+						.Justification(ETextJustify::Center)
+						.Value(this, &STakeRecorderCockpit::GetTakeNumber)
+						.Delta(1)
+						.MinValue(1)
+						.MaxValue(TOptional<int32>())
+						.OnBeginSliderMovement(this, &STakeRecorderCockpit::OnBeginSetTakeNumber)
+						.OnValueChanged(this, &STakeRecorderCockpit::SetTakeNumber)
+						.OnValueCommitted(this, &STakeRecorderCockpit::SetTakeNumber_FromCommit)
+						.OnEndSliderMovement(this, &STakeRecorderCockpit::OnEndSetTakeNumber)
+						.TypeInterface(DigitsTypeInterface)
+					]
+				];
+			}
+			// Custom user keys
+			else
+			{
+				UserTokensBox->AddSlot()
+				.HAlign(HAlign_Fill)
+				.MinWidth(MinWidth)
+				.MaxWidth(MaxWidth)
+				.AutoWidth()
+				[
+					SNew(SHorizontalBox)
+					+ SHorizontalBox::Slot()
+					[
+						SNew(SOverlay)
+
+						+ SOverlay::Slot()
+						.VAlign(VAlign_Top)
+						.HAlign(HAlign_Left)
+						.Padding(2.0f, 2.0f)
+						[
+							SNew(SHorizontalBox)
+							+ SHorizontalBox::Slot()
+							.AutoWidth()
+							[
+								SNew(STextBlock)
+								.TextStyle(FTakeRecorderStyle::Get(), "TakeRecorder.TextBox")
+								.Text(FText::FromString(TokenString.ToUpper()))
+							]
+						]
+						+ SOverlay::Slot()
+						.VAlign(VAlign_Center)
+						.HAlign(HAlign_Center)
+						[
+							SNew(SEditableTextBox)
+							.IsEnabled(this, &STakeRecorderCockpit::EditingMetaData)
+							.Style(FTakeRecorderStyle::Get(), "TakeRecorder.EditableTextBox")
+							.Font(FTakeRecorderStyle::Get().GetFontStyle("TakeRecorder.Cockpit.LargeText"))
+							.HintText(LOCTEXT("EnterToken_Hint", "<value>"))
+							.Justification(ETextJustify::Left)
+							.SelectAllTextWhenFocused(true)
+							.Text(this, &STakeRecorderCockpit::GetCustomTokenTextValue, TokenString)
+							.OnTextCommitted(this, &STakeRecorderCockpit::SetCustomTokenTextValue, TokenString)
+						]
+					]
+				];
+			}
+		}
+	}
+}
+
 #undef LOCTEXT_NAMESPACE
diff --git a/Engine/Plugins/VirtualProduction/Takes/Source/TakeRecorder/Private/Widgets/STakeRecorderCockpit.h b/Engine/Plugins/VirtualProduction/Takes/Source/TakeRecorder/Private/Widgets/STakeRecorderCockpit.h
--- a/Engine/Plugins/VirtualProduction/Takes/Source/TakeRecorder/Private/Widgets/STakeRecorderCockpit.h
+++ b/Engine/Plugins/VirtualProduction/Takes/Source/TakeRecorder/Private/Widgets/STakeRecorderCockpit.h
@@ -2,19 +2,27 @@
 
 #pragma once
 
+#include "Recorder/TakeRecorderParameters.h"
+#include "UObject/GCObject.h"
+#include "UObject/WeakFieldPtr.h"
 #include "Widgets/DeclarativeSyntaxSupport.h"
 #include "Widgets/SCompoundWidget.h"
-#include "UObject/GCObject.h"
-#include "Recorder/TakeRecorderParameters.h"
+#include "Widgets/TakeRecorderNamingTokensWidgetData.h"
 
 enum class ECheckBoxState : uint8;
 enum class ETakeRecordMode : uint8;
 
 struct FFrameRate;
 class ULevelSequence;
+class UNamingTokens;
 class UTakeMetaData;
 class UTakeRecorder;
+class SHorizontalBox;
+class IDetailsView;
 struct FDigitsTypeInterface;
+struct FNamingTokenResultData;
+struct FNamingTokensEvaluationData;
+struct FPropertyChangedEvent;
 
 /**
  * Cockpit UI for defining take meta-data.
@@ -65,6 +73,11 @@ public:
 
 	void Refresh();
 
+	/** Externally called when a property has been updated. */
+	void NotifyPropertyUpdated(const FPropertyChangedEvent& InPropertyChangedEvent);
+	/** Externally called when a details view has been added related to this take recorder. */
+	void NotifyDetailsViewAdded(const TWeakPtr<IDetailsView>& InDetailsView);
+
 private:
 
 	virtual void AddReferencedObjects(FReferenceCollector& Collector) override;
@@ -74,6 +87,8 @@ private:
 	}
 	virtual void Tick(const FGeometry& AllottedGeometry, const double InCurrentTime, const float InDeltaTime) override;
 
+	void PostUndoRedo();
+	
 private:
 
 	EVisibility GetTakeWarningVisibility() const;
@@ -106,6 +121,9 @@ private:
 	FText GetTimestampText() const;
 	FText GetTimestampTooltipText() const;
 
+	/** Fully evaluated text of the take save dir. */
+	FText GetEvaluatedTakeSaveDirText() const;
+
 	int32 GetTakeNumber() const;
 
 	FReply OnSetNextTakeNumber();
@@ -150,6 +168,45 @@ private:
 	
 	TSharedRef<SWidget> OnCreateMenu();
 
+	/** Setup bindings to our naming tokens. */
+	void BindToNamingTokenEvents();
+	/** Remove bindings to our naming tokens. */
+	void UnbindNamingTokensEvents();
+
+	/** Find a user defined text value for our custom tokens. */
+	FText GetCustomTokenTextValue(FString InTokenKey) const;
+	/** Set a user defined text value for a custom token. */
+	void SetCustomTokenTextValue(const FText& InNewText, ETextCommit::Type InCommitType, FString InTokenKey);
+
+	/** When a token value has been updated by the user, such as for slate, take, or custom tokens. */
+	void OnTokenValueUpdated();
+
+	/** Evaluate the take dir tokens for our example display. */
+	void EvaluateTakeSaveDirTokens();
+	
+	/** Called before our naming tokens class evaluates. Used to populate naming token data with our custom token definitions. */
+	void OnTakeRecorderNamingTokensPreEvaluate(const FNamingTokensEvaluationData& InEvaluationData);
+
+	/** Iterate all cached details views refreshing all undefined tokens. */
+	void RefreshUndefinedTokens();
+	
+	/** Process all properties in the details view. */
+	void UpdateUndefinedTokensFromDetailsView(const TWeakPtr<IDetailsView>& InDetailsView);
+
+	/** Create and use temporary details view for key data for the case where no details views are available. */
+	void UpdateUndefinedTokensFromTemporaryDetailsView();
+	
+	/**
+	 * Evaluate naming tokens given a property.
+	 * @param InProperty The property we need the value from.
+	 * @param InContainer Container of the property which holds the value.
+	 * @param bForce Whether to force an evaluation even if the property isn't flagged as a NamingToken. Property still must be a supported type.
+	 */
+	void EvaluateTokensFromProperty(const FProperty* InProperty, const void* InContainer, bool bForce = false);
+
+	/** Configure any user tokens in the UI. */
+	void CreateUserTokensUI();
+	
 private:
 
 	/** Take meta-data cached from the level sequence if it exists. Referenced by AddReferencedObjects. */
@@ -179,7 +236,30 @@ private:
 
 	TSharedPtr<FDigitsTypeInterface> DigitsTypeInterface;
 
-	// Cached take numbers and slate used to UpdateTakeError() only when necesary
+	/** All known detail views. */
+	TArray<TWeakPtr<IDetailsView>> DetailViews;
+
+	/** Details views we manage so we know what tokens are available in case the user has hidden the normal details. */
+	TArray<TSharedPtr<IDetailsView>> TemporaryDetailsViews;
+	
+	/** Signal that undefined tokens should be refreshed on the next tick. */
+	bool bRefreshUndefinedTokens = false;
+	
+	/** Tokens added by the user */
+	TSharedPtr<SHorizontalBox> UserTokensBox;
+
+	// Cached take numbers and slate used to UpdateTakeError() only when necessary
 	int32 CachedTakeNumber;
 	FString CachedTakeSlate;
+
+	/** Our cached Take Save Dir Property. */
+	TWeakFieldPtr<const FProperty> CachedTakeSaveDirProperty;
+	/** Our cached object container, holding our Take Save Dir property value. */
+	const void* CachedTakeSaveDirContainer = nullptr;
+
+	/** Stored data relevant to our naming tokens. Managed as UObject for transactions. */
+	TStrongObjectPtr<UTakeRecorderNamingTokensWidgetData> NamingTokensData;
+
+	/** Handle for binding to the post undo/redo editor delegate. */
+	FDelegateHandle PostUndoRedoDelegateHandle;
 };
diff --git a/Engine/Plugins/VirtualProduction/Takes/Source/TakeRecorder/Private/Widgets/STakeRecorderPanel.cpp b/Engine/Plugins/VirtualProduction/Takes/Source/TakeRecorder/Private/Widgets/STakeRecorderPanel.cpp
--- a/Engine/Plugins/VirtualProduction/Takes/Source/TakeRecorder/Private/Widgets/STakeRecorderPanel.cpp
+++ b/Engine/Plugins/VirtualProduction/Takes/Source/TakeRecorder/Private/Widgets/STakeRecorderPanel.cpp
@@ -145,7 +145,9 @@ void STakeRecorderPanel::Construct(const FArguments& InArgs)
 	.TakeRecorderMode(this, &STakeRecorderPanel::GetTakeRecorderMode);
 
 	LevelSequenceTakeWidget = SNew(SLevelSequenceTakeEditor)
-	.LevelSequence(this, &STakeRecorderPanel::GetLevelSequence);
+	.LevelSequence(this, &STakeRecorderPanel::GetLevelSequence)
+	.OnDetailsPropertiesChanged(this, &STakeRecorderPanel::OnLevelSequenceDetailsChanged)
+	.OnDetailsViewAdded(this, &STakeRecorderPanel::OnLevelSequenceDetailsViewAdded);
 
 	// Create the sequencer panel, and open it if necessary
 	SequencerPanel = MakeShared<FScopedSequencerPanel>(MakeAttributeSP(this, &STakeRecorderPanel::GetLevelSequence));
@@ -1002,6 +1004,22 @@ void STakeRecorderPanel::OnLevelSequenceChanged()
 	RefreshPanel();
 }
 
+void STakeRecorderPanel::OnLevelSequenceDetailsChanged(const FPropertyChangedEvent& InPropertyChangedEvent)
+{
+	if (CockpitWidget.IsValid())
+	{
+		CockpitWidget->NotifyPropertyUpdated(InPropertyChangedEvent);
+	}
+}
+
+void STakeRecorderPanel::OnLevelSequenceDetailsViewAdded(const TWeakPtr<IDetailsView>& InDetailsView)
+{
+	if (CockpitWidget.IsValid())
+	{
+		CockpitWidget->NotifyDetailsViewAdded(InDetailsView);
+	}
+}
+
 void STakeRecorderPanel::SetRecordIntoLevelSequence(ULevelSequence* LevelSequence)
 {
 	SuppliedLevelSequence = nullptr;
diff --git a/Engine/Plugins/VirtualProduction/Takes/Source/TakeRecorder/Private/Widgets/STakeRecorderPanel.h b/Engine/Plugins/VirtualProduction/Takes/Source/TakeRecorder/Private/Widgets/STakeRecorderPanel.h
--- a/Engine/Plugins/VirtualProduction/Takes/Source/TakeRecorder/Private/Widgets/STakeRecorderPanel.h
+++ b/Engine/Plugins/VirtualProduction/Takes/Source/TakeRecorder/Private/Widgets/STakeRecorderPanel.h
@@ -125,6 +125,10 @@ private:
 	void ToggleSettings(ECheckBoxState CheckState);
 
 	void OnLevelSequenceChanged();
+	/** When properties in a level sequence details have been changed. */
+	void OnLevelSequenceDetailsChanged(const FPropertyChangedEvent& InPropertyChangedEvent);
+	/** When a level sequence has a details view added. */
+	void OnLevelSequenceDetailsViewAdded(const TWeakPtr<IDetailsView>& InDetailsView);
 
 	ECheckBoxState GetTakeBrowserCheckState() const;
 	void ToggleTakeBrowserCheckState(ECheckBoxState CheckState);
diff --git a/Engine/Plugins/VirtualProduction/Takes/Source/TakeRecorder/Private/Widgets/TakeRecorderNamingTokensWidgetData.cpp b/Engine/Plugins/VirtualProduction/Takes/Source/TakeRecorder/Private/Widgets/TakeRecorderNamingTokensWidgetData.cpp
new file mode 100644
--- /dev/null
+++ b/Engine/Plugins/VirtualProduction/Takes/Source/TakeRecorder/Private/Widgets/TakeRecorderNamingTokensWidgetData.cpp
@@ -0,0 +1,30 @@
+// Copyright Epic Games, Inc. All Rights Reserved.
+
+#include "TakeRecorderNamingTokensWidgetData.h"
+
+TArray<FString>& UTakeRecorderNamingTokensWidgetData::FindOrAddTokenKeysForField(const FString& InFieldName)
+{
+	if (FTakeRecorderNamingTokensFieldMapping* Match = FieldToUndefinedKeys.FindByPredicate([&InFieldName](const FTakeRecorderNamingTokensFieldMapping& Fields)
+	{
+		return Fields.FieldName == InFieldName;
+	}))
+	{
+		return Match->UndefinedKeys;
+	}
+	FTakeRecorderNamingTokensFieldMapping Data;
+	Data.FieldName = InFieldName;
+	return FieldToUndefinedKeys.Add_GetRef(Data).UndefinedKeys;
+}
+
+bool UTakeRecorderNamingTokensWidgetData::IsTokenKeyUndefined(const FString& InTokenKey) const
+{
+	for (const FTakeRecorderNamingTokensFieldMapping& TokenFieldMapping : FieldToUndefinedKeys)
+	{
+		if (TokenFieldMapping.UndefinedKeys.Contains(InTokenKey))
+		{
+			return true;
+		}
+	}
+
+	return false;
+}
diff --git a/Engine/Plugins/VirtualProduction/Takes/Source/TakeRecorder/Private/Widgets/TakeRecorderNamingTokensWidgetData.h b/Engine/Plugins/VirtualProduction/Takes/Source/TakeRecorder/Private/Widgets/TakeRecorderNamingTokensWidgetData.h
new file mode 100644
--- /dev/null
+++ b/Engine/Plugins/VirtualProduction/Takes/Source/TakeRecorder/Private/Widgets/TakeRecorderNamingTokensWidgetData.h
@@ -0,0 +1,63 @@
+// Copyright Epic Games, Inc. All Rights Reserved.
+
+#pragma once
+
+#include "CoreMinimal.h"
+
+#include "TakeRecorderNamingTokensWidgetData.generated.h"
+
+class UNamingTokens;
+
+USTRUCT()
+struct FTakeRecorderNamingTokensFieldMapping
+{
+	GENERATED_BODY()
+
+	/** Name of the property field. */
+	UPROPERTY()
+	FString FieldName;
+	
+	/** All undefined keys with this field. */
+	UPROPERTY()
+	TArray<FString> UndefinedKeys;
+};
+
+/**
+ * Data container object for the Take Recorder widget. Kept as UObject to assist with transactions.
+ */
+UCLASS()
+class UTakeRecorderNamingTokensWidgetData : public UObject
+{
+	GENERATED_BODY()
+public:
+	/** Guid to our managed tokens. */
+	FGuid NamingTokensExternalGuid;
+	
+	/** Custom tokens entered in by the user, mapped to a user defined value. */
+	UPROPERTY()
+	TMap<FString, FText> UserDefinedTokens;
+
+	/**
+	 * User tokens that are currently visible. Kept as a separate property from UsedDefinedTokens we can persist user values
+	 * between selected sources that have different token entries in their fields.
+	 */
+	UPROPERTY(Transient)
+	TSet<FString> VisibleUserTokens;
+	
+	/** Ordered array of FieldName to undefined token keys. These are present if a token evaluation failed to identify them. */
+	UPROPERTY(NonTransactional)
+	TArray<FTakeRecorderNamingTokensFieldMapping> FieldToUndefinedKeys;
+
+	/** Evaluated text to serve as an example. */
+	UPROPERTY()
+	FText EvaluatedTextValue;
+	
+	/** Pointer to our naming tokens object. */
+	TWeakObjectPtr<UNamingTokens> TakeRecorderNamingTokens;
+
+	/** Find or add an array given a field name. The array will be a list of token keys for this specific field. */
+	TArray<FString>& FindOrAddTokenKeysForField(const FString& InFieldName);
+
+	/** Checks if a token key is in our undefined list. */
+	bool IsTokenKeyUndefined(const FString& InTokenKey) const;
+};
\ No newline at end of file
diff --git a/Engine/Plugins/VirtualProduction/Takes/Source/TakeRecorder/Public/Recorder/TakeRecorderParameters.h b/Engine/Plugins/VirtualProduction/Takes/Source/TakeRecorder/Public/Recorder/TakeRecorderParameters.h
--- a/Engine/Plugins/VirtualProduction/Takes/Source/TakeRecorder/Public/Recorder/TakeRecorderParameters.h
+++ b/Engine/Plugins/VirtualProduction/Takes/Source/TakeRecorder/Public/Recorder/TakeRecorderParameters.h
@@ -67,7 +67,7 @@ struct FTakeRecorderProjectParameters
 	TAKERECORDER_API FString GetTakeAssetPath() const;
 
 	/** The root of the directory in which to save recorded takes. */
-	UPROPERTY(config, EditAnywhere, BlueprintReadWrite, Category = "Take Recorder", meta = (ContentDir))
+	UPROPERTY(config, EditAnywhere, BlueprintReadWrite, Category = "Take Recorder", meta = (ContentDir, NamingTokens))
 	FDirectoryPath RootTakeSaveDir;
 
 	/**
@@ -81,7 +81,7 @@ struct FTakeRecorderProjectParameters
 	 * {take}      - The take number.
 	 * {slate}     - The slate string.
 	 */
-	UPROPERTY(config, EditAnywhere, BlueprintReadWrite, Category = "Take Recorder")
+	UPROPERTY(config, EditAnywhere, BlueprintReadWrite, Category = "Take Recorder", meta = (NamingTokens))
 	FString TakeSaveDir;
 
 	/**
diff --git a/Engine/Plugins/VirtualProduction/Takes/Source/TakeRecorder/Public/TakeRecorderSettings.h b/Engine/Plugins/VirtualProduction/Takes/Source/TakeRecorder/Public/TakeRecorderSettings.h
--- a/Engine/Plugins/VirtualProduction/Takes/Source/TakeRecorder/Public/TakeRecorderSettings.h
+++ b/Engine/Plugins/VirtualProduction/Takes/Source/TakeRecorder/Public/TakeRecorderSettings.h
@@ -34,7 +34,7 @@ public:
 	FTakeRecorderUserParameters Settings;
 
 	/** The default location in which to save take presets */
-	UPROPERTY(config, EditAnywhere, Category="User Settings", DisplayName="Preset Save Location", meta = (ContentDir))
+	UPROPERTY(config, EditAnywhere, Category="User Settings", DisplayName="Preset Save Location", meta = (ContentDir, NamingTokens))
 	FDirectoryPath PresetSaveDir;
 
 	/** Soft reference to the preset last opened on the take recording UI */
diff --git a/Engine/Plugins/VirtualProduction/Takes/Source/TakeRecorder/TakeRecorder.Build.cs b/Engine/Plugins/VirtualProduction/Takes/Source/TakeRecorder/TakeRecorder.Build.cs
--- a/Engine/Plugins/VirtualProduction/Takes/Source/TakeRecorder/TakeRecorder.Build.cs
+++ b/Engine/Plugins/VirtualProduction/Takes/Source/TakeRecorder/TakeRecorder.Build.cs
@@ -35,7 +35,8 @@ public class TakeRecorder : ModuleRules
 				"WorkspaceMenuStructure",
 				"Analytics",
 				"ToolMenus",
-				"ToolWidgets"
+				"ToolWidgets",
+				"NamingTokens"
 			}
 		);
 
@@ -49,5 +50,12 @@ public class TakeRecorder : ModuleRules
 				"TraceLog"
             }
         );
-    }
+
+		PublicIncludePathModuleNames.AddRange(
+			new string[]
+			{
+				"TakeRecorderNamingTokens"
+			}
+		);
+	}
 }
diff --git a/Engine/Plugins/VirtualProduction/Takes/Source/TakeRecorderNamingTokens/Private/TakeRecorderNamingTokens.cpp b/Engine/Plugins/VirtualProduction/Takes/Source/TakeRecorderNamingTokens/Private/TakeRecorderNamingTokens.cpp
--- a/Engine/Plugins/VirtualProduction/Takes/Source/TakeRecorderNamingTokens/Private/TakeRecorderNamingTokens.cpp
+++ b/Engine/Plugins/VirtualProduction/Takes/Source/TakeRecorderNamingTokens/Private/TakeRecorderNamingTokens.cpp
@@ -6,7 +6,9 @@
 #include "ITakeRecorderNamingTokensModule.h"
 #include "Misc/PackageName.h"
 #include "MovieSceneToolsProjectSettings.h"
+#include "NamingTokens/TakeRecorderNamingTokensContext.h"
 #include "TakeMetaData.h"
+#include "TakeRecorderNamingTokensLog.h"
 #include "TakeRecorderSettings.h"
 #include "TakesCoreBlueprintLibrary.h"
 
@@ -105,13 +107,49 @@ void UTakeRecorderNamingTokens::OnCreateDefaultTokens(TArray<FNamingTokenData>&
 	        return FText::FromString(MapName);
 	    })
 	});
+
+	Tokens.Add({
+		TEXT("actor"),
+		LOCTEXT("TokenActor", "Actor"),
+		FNamingTokenData::FTokenProcessorDelegateNative::CreateLambda([this]{
+			if (Context && Context->Actor)
+			{
+				return FText::FromString(Context->Actor->GetActorLabel());
+			}
+
+			UE_LOG(LogTakeRecorderNamingTokens, Verbose, TEXT("Attempted to use 'actor' naming token but no context is available."));
+
+			return FText::GetEmpty();
+		})
+	});
+	
+	Tokens.Add({
+		TEXT("channel"),
+		LOCTEXT("TokenChannel", "Channel"),
+		FNamingTokenData::FTokenProcessorDelegateNative::CreateLambda([this]{
+			if (Context)
+			{
+				return FText::FromString(FString::FromInt(Context->AudioInputDeviceChannel));
+			}
+
+			UE_LOG(LogTakeRecorderNamingTokens, Verbose, TEXT("Attempted to use 'channel' naming token but no context is available."));
+
+			return FText::GetEmpty();
+		})
+	});
 }
 
 void UTakeRecorderNamingTokens::OnPreEvaluate_Implementation(const FNamingTokensEvaluationData& InEvaluationData)
 {
 	Super::OnPreEvaluate_Implementation(InEvaluationData);
-	
-	TakeMetaData = UTakeMetaData::GetMostRecentMetaData();
+
+	UTakeRecorderNamingTokensContext* MatchingContext = nullptr;
+	InEvaluationData.Contexts.FindItemByClass<UTakeRecorderNamingTokensContext>(&MatchingContext);
+	Context = MatchingContext;
+
+	// Use either provided take meta-data (manually evaluated from take recorder), or locate the most recent (global operation)
+	TakeMetaData = (Context && Context->TakeMetaData != nullptr) ?
+		Context->TakeMetaData : TObjectPtr<UTakeMetaData>(UTakeMetaData::GetMostRecentMetaData());
 	if (TakeMetaData == nullptr)
 	{
 		// Create metadata which computes its data from available information.
@@ -119,21 +157,23 @@ void UTakeRecorderNamingTokens::OnPreEvaluate_Implementation(const FNamingTokens
 		// method would be available for both us and the slate widget, but UTakeRecorderProjectSettings creates
 		// a circular dependency we should avoid.
 		
-		TakeMetaData = UTakeMetaData::CreateFromDefaults(GetTransientPackage(), NAME_None);
-		TakeMetaData->SetFlags(RF_Transactional | RF_Transient);
+		UTakeMetaData* TransientTakeMetaData = UTakeMetaData::CreateFromDefaults(GetTransientPackage(), NAME_None);
+		TransientTakeMetaData->SetFlags(RF_Transactional | RF_Transient);
 
 		const FString DefaultSlate = GetDefault<UTakeRecorderProjectSettings>()->Settings.DefaultSlate;
-		if (TakeMetaData->GetSlate() != DefaultSlate)
+		if (TransientTakeMetaData->GetSlate() != DefaultSlate)
 		{
-			TakeMetaData->SetSlate(DefaultSlate, false);
+			TransientTakeMetaData->SetSlate(DefaultSlate, false);
 		}
 
 		// Compute the correct starting take number
-		const int32 NextTakeNumber = UTakesCoreBlueprintLibrary::ComputeNextTakeNumber(TakeMetaData->GetSlate());
-		if (TakeMetaData->GetTakeNumber() != NextTakeNumber)
+		const int32 NextTakeNumber = UTakesCoreBlueprintLibrary::ComputeNextTakeNumber(TransientTakeMetaData->GetSlate());
+		if (TransientTakeMetaData->GetTakeNumber() != NextTakeNumber)
 		{
-			TakeMetaData->SetTakeNumber(NextTakeNumber, false);
+			TransientTakeMetaData->SetTakeNumber(NextTakeNumber, false);
 		}
+
+		TakeMetaData = TransientTakeMetaData;
 	}
 }
 
@@ -142,7 +182,8 @@ void UTakeRecorderNamingTokens::OnPostEvaluate_Implementation()
 	Super::OnPostEvaluate_Implementation();
 
 	// Make sure we don't keep a strong reference in case we're linking to active metadata.
-	TakeMetaData = nullptr;
+    TakeMetaData = nullptr;
+	Context = nullptr;
 }
 
 FDateTime UTakeRecorderNamingTokens::GetCurrentDateTime_Implementation() const
diff --git a/Engine/Plugins/VirtualProduction/Takes/Source/TakeRecorderNamingTokens/Private/TakeRecorderNamingTokens.h b/Engine/Plugins/VirtualProduction/Takes/Source/TakeRecorderNamingTokens/Private/TakeRecorderNamingTokens.h
--- a/Engine/Plugins/VirtualProduction/Takes/Source/TakeRecorderNamingTokens/Private/TakeRecorderNamingTokens.h
+++ b/Engine/Plugins/VirtualProduction/Takes/Source/TakeRecorderNamingTokens/Private/TakeRecorderNamingTokens.h
@@ -7,6 +7,7 @@
 #include "TakeRecorderNamingTokens.generated.h"
 
 class UTakeMetaData;
+class UTakeRecorderNamingTokensContext;
 
 /**
  * Naming Tokens for Take Recorder.
@@ -31,5 +32,9 @@ protected:
 private:
 	/** Cached metadata for this run. */
 	UPROPERTY(Transient)
-	TObjectPtr<UTakeMetaData> TakeMetaData;
+	TObjectPtr<const UTakeMetaData> TakeMetaData;
+
+	/** Cached context for this run. This isn't available globally and requires a context passed to it. */
+	UPROPERTY(Transient)
+	TObjectPtr<UTakeRecorderNamingTokensContext> Context;
 };
diff --git a/Engine/Plugins/VirtualProduction/Takes/Source/TakeRecorderSources/Private/TakeRecorderMicrophoneAudioSource.cpp b/Engine/Plugins/VirtualProduction/Takes/Source/TakeRecorderSources/Private/TakeRecorderMicrophoneAudioSource.cpp
--- a/Engine/Plugins/VirtualProduction/Takes/Source/TakeRecorderSources/Private/TakeRecorderMicrophoneAudioSource.cpp
+++ b/Engine/Plugins/VirtualProduction/Takes/Source/TakeRecorderSources/Private/TakeRecorderMicrophoneAudioSource.cpp
@@ -1,11 +1,13 @@
 // Copyright Epic Games, Inc. All Rights Reserved.
 
 #include "TakeRecorderMicrophoneAudioSource.h"
+
+#include "NamingTokens/TakeRecorderNamingTokensContext.h"
+#include "TakeMetaData.h"
 #include "TakeRecorderMicrophoneAudioManager.h"
-#include "TakeRecorderSources.h"
 #include "TakeRecorderSettings.h"
+#include "TakeRecorderSources.h"
 #include "TakesUtils.h"
-#include "TakeMetaData.h"
 
 #include "AudioCaptureEditorTypes.h"
 #include "Editor.h"
@@ -53,7 +55,10 @@ FString UTakeRecorderMicrophoneAudioSourceSettings::GetSubsceneTrackName(ULevelS
 {
 	if (UTakeMetaData* TakeMetaData = InSequence->FindMetaData<UTakeMetaData>())
 	{
-		return FString::Printf(TEXT("%s_%s"), *AudioTrackName.ToString(), *TakeMetaData->GenerateAssetPath("{slate}"));
+		// Before Naming Tokens we were always evaluating AudioTrackName_{slate}.
+		// @todo NamingTokens - Keep it consistent, but maybe subscene track name should become its own user defined field which includes {slate} by default?
+		const FString SubSceneTrackName = FString::Printf(TEXT("%s_%s"), *AudioTrackName.ToString(), *TakeMetaData->GenerateAssetPath("{slate}"));
+		return *TakeMetaData->GenerateAssetPath(SubSceneTrackName);
 	}
 	return TEXT("MicrophoneAudio");
 }
@@ -148,21 +153,12 @@ void UTakeRecorderMicrophoneAudioSource::SetAudioDeviceChannelCount(int32 InChan
 	}
 }
 
-FString UTakeRecorderMicrophoneAudioSource::ReplaceStringTokens(const FString& InString) const
-{
-	TMap<FString, FStringFormatArg> FormatArgs;
-	FormatArgs.Add(TEXT("channel"), FString::FromInt(AudioChannel.AudioInputDeviceChannel));
-
-	return FString::Format(*InString, FormatArgs);
-}
-
 FString UTakeRecorderMicrophoneAudioSource::GetAudioTrackName(ULevelSequence* InSequence) const
 {
-	if (UTakeMetaData* TakeMetaData = InSequence->FindMetaData<UTakeMetaData>())
+	FString EvaluatedAudioTrackName;
+	if (TryEvaluateTokenString(AudioTrackName.ToString(), InSequence, EvaluatedAudioTrackName))
 	{
-		FString TempStr = *TakeMetaData->GenerateAssetPath(*AudioTrackName.ToString());
-
-		return ReplaceStringTokens(TempStr);
+		return EvaluatedAudioTrackName;
 	}
 
 	return TEXT("MicrophoneAudio");
@@ -173,9 +169,10 @@ FString UTakeRecorderMicrophoneAudioSource::GetAudioAssetName(ULevelSequence* In
 	// This is called for both uassets and subsequences, both of which require unique names.
 	// Append the channel number if it's not already part of the name to ensure uniqueness.
 
-	if (UTakeMetaData* TakeMetaData = InSequence->FindMetaData<UTakeMetaData>())
+	FString EvaluatedAudioAssetName;
+	if (TryEvaluateTokenString(AudioAssetName, InSequence, EvaluatedAudioAssetName))
 	{
-		return UPackageTools::SanitizePackageName(ReplaceStringTokens(TakeMetaData->GenerateAssetPath(AudioAssetName)));
+		return UPackageTools::SanitizePackageName(EvaluatedAudioAssetName);
 	}
 
 	return TEXT("MicrophoneAudio");
@@ -280,9 +277,9 @@ TArray<UTakeRecorderSource*> UTakeRecorderMicrophoneAudioSource::PreRecording(UL
 		if (AudioSubDirectory.Len())
 		{
 			FString TempStr = AudioSubDirectory;
-			if (UTakeMetaData* TakeMetaData = InSequence->FindMetaData<UTakeMetaData>())
+			if (TryEvaluateTokenString(TempStr, InSequence, TempStr))
 			{
-				TempStr = UPackageTools::SanitizePackageName(ReplaceStringTokens(TakeMetaData->GenerateAssetPath(TempStr)));
+				TempStr = UPackageTools::SanitizePackageName(TempStr);
 			}
 			AudioDirectory.Path /= TempStr;
 		}
@@ -497,7 +494,12 @@ void UTakeRecorderMicrophoneAudioSource::GetRecordedSoundWave(ULevelSequence* In
 
 FText UTakeRecorderMicrophoneAudioSource::GetDisplayTextImpl() const
 {
-	return FText::FromString(ReplaceStringTokens(AudioSourceName.ToString()));
+	// We purposely don't use naming tokens here to avoid evaluating all tokens per tick. It's unlikely we need
+	// to support all tokens for AudioSourceName.
+	FFormatNamedArguments FormatArgs;
+	FormatArgs.Add(TEXT("channel"), FText::FromString(FString::FromInt(AudioChannel.AudioInputDeviceChannel)));
+	
+	return FText::Format(AudioSourceName, FormatArgs);
 }
 
 FText UTakeRecorderMicrophoneAudioSource::GetAddSourceDisplayTextImpl() const
@@ -570,3 +572,15 @@ FString UTakeRecorderMicrophoneAudioSource::CreateUniqueAudioAssetName(ULevelSeq
 	
 	return NewAssetName;
 }
+
+bool UTakeRecorderMicrophoneAudioSource::TryEvaluateTokenString(const FString& InString, const ULevelSequence* InSequence, FString& OutString) const
+{
+	if (const UTakeMetaData* TakeMetaData = InSequence->FindMetaData<UTakeMetaData>())
+	{
+		UTakeRecorderNamingTokensContext* Context = NewObject<UTakeRecorderNamingTokensContext>();
+		Context->AudioInputDeviceChannel = AudioChannel.AudioInputDeviceChannel;
+		OutString = TakeMetaData->GenerateAssetPath(InString, Context);
+		return true;
+	}
+	return false;
+}
diff --git a/Engine/Plugins/VirtualProduction/Takes/Source/TakeRecorderSources/Private/TakeRecorderMicrophoneAudioSource.h b/Engine/Plugins/VirtualProduction/Takes/Source/TakeRecorderSources/Private/TakeRecorderMicrophoneAudioSource.h
--- a/Engine/Plugins/VirtualProduction/Takes/Source/TakeRecorderSources/Private/TakeRecorderMicrophoneAudioSource.h
+++ b/Engine/Plugins/VirtualProduction/Takes/Source/TakeRecorderSources/Private/TakeRecorderMicrophoneAudioSource.h
@@ -38,7 +38,7 @@ public:
 	FText AudioSourceName;
 
 	/** Name of the recorded audio track */
-	UPROPERTY(config, EditAnywhere, BlueprintReadWrite, Category = "Source")
+	UPROPERTY(config, EditAnywhere, BlueprintReadWrite, Category = "Source", meta = (NamingTokens))
 	FText AudioTrackName;
 
 	/** The name of the audio asset.
@@ -52,7 +52,7 @@ public:
 	 * {take} - The take number.
 	 * {slate} - The slate string.
 	 */
-	UPROPERTY(config, EditAnywhere, BlueprintReadWrite, Category = "Source")
+	UPROPERTY(config, EditAnywhere, BlueprintReadWrite, Category = "Source", meta = (NamingTokens))
 	FString AudioAssetName;
 
 	/** The name of the subdirectory audio will be placed in. Leave this empty to place into the same directory as the sequence base path 
@@ -66,7 +66,7 @@ public:
 	 * {take} - The take number.
 	 * {slate} - The slate string.
 	 */
-	UPROPERTY(config, EditAnywhere, BlueprintReadWrite, Category = "Source")
+	UPROPERTY(config, EditAnywhere, BlueprintReadWrite, Category = "Source", meta = (NamingTokens))
 	FString AudioSubDirectory;
 };
 
@@ -136,9 +136,6 @@ private:
 	/** Parses recorded timecode values */
 	void ProcessRecordedTimes(ULevelSequence* InSequence);
 
-	/** Calls ReplaceSourceTokens and ReplaceBaseTokens in an effort to expand all embedded tokens in the given string */
-	FString ReplaceStringTokens(const FString& InString) const;
-
 	/** Returns an array of booleans indicating which channel indexes are currently in use */
 	TUniquePtr<TArray<bool>> GetChannelsInUse(const int32 InDeviceChannelCount);
 	/** Fetches the USoundWave for this source after a Take has been recorded */
@@ -155,6 +152,9 @@ private:
 	/** Static, parameterized helper function for buidling unique asset names */	
 	static FString CreateUniqueAudioAssetName(ULevelSequence* InSequence, UTakeRecorderSources* InSources, const FString& InAssetName, const int32 InChannelNumber);
 
+	/** Evaluate our tokens and automatically create an appropriate context. */
+	bool TryEvaluateTokenString(const FString& InString, const ULevelSequence* InSequence, FString& OutString) const;
+	
 private:
 
 	// Holds the Sequencer audio track for this source
diff --git a/Engine/Plugins/VirtualProduction/Takes/Source/TakeTrackRecorders/Private/MovieSceneAnimationTrackRecorder.cpp b/Engine/Plugins/VirtualProduction/Takes/Source/TakeTrackRecorders/Private/MovieSceneAnimationTrackRecorder.cpp
--- a/Engine/Plugins/VirtualProduction/Takes/Source/TakeTrackRecorders/Private/MovieSceneAnimationTrackRecorder.cpp
+++ b/Engine/Plugins/VirtualProduction/Takes/Source/TakeTrackRecorders/Private/MovieSceneAnimationTrackRecorder.cpp
@@ -23,6 +23,8 @@
 #include "UObject/UObjectBaseUtility.h"
 #include "ObjectTools.h"
 #include "Editor.h"
+#include "NamingTokenData.h"
+#include "NamingTokens/TakeRecorderNamingTokensContext.h"
 
 #include UE_INLINE_GENERATED_CPP_BY_NAME(MovieSceneAnimationTrackRecorder)
 
@@ -64,14 +66,11 @@ void UMovieSceneAnimationTrackRecorder::CreateAnimationAssetAndSequence(const AA
 		{
 			UTakeMetaData* AssetMetaData = LeaderLevelSequence->FindMetaData<UTakeMetaData>();
 
-			AnimationAssetName = AssetMetaData->GenerateAssetPath(AnimSettings->AnimationAssetName);
-			DirectoryPath = AssetMetaData->GenerateAssetPath(DirectoryPath);
-
-			TMap<FString, FStringFormatArg> FormatArgs;
-			FormatArgs.Add(TEXT("actor"), Actor->GetActorLabel());
-
-			AnimationAssetName = FString::Format(*AnimationAssetName, FormatArgs);
-			DirectoryPath = FString::Format(*DirectoryPath, FormatArgs);
+			UTakeRecorderNamingTokensContext* Context = NewObject<UTakeRecorderNamingTokensContext>();
+			Context->Actor = Actor;
+			
+			AnimationAssetName = AssetMetaData->GenerateAssetPath(AnimSettings->AnimationAssetName, Context);
+			DirectoryPath = AssetMetaData->GenerateAssetPath(DirectoryPath, Context);
 		}
 
 		AnimSequence = TakesUtils::MakeNewAsset<UAnimSequence>(DirectoryPath, AnimationAssetName);
@@ -190,7 +189,18 @@ void UMovieSceneAnimationTrackRecorder::CreateTrackImpl()
 
 			if (AnimTrack)
 			{
-				AnimTrack->SetDisplayName(AnimSettings->AnimationTrackName);
+				FText AnimationTrackDisplayName = AnimSettings->AnimationTrackName;
+				if (const ULevelSequence* LeaderLevelSequence = OwningTakeRecorderSource->GetRootLevelSequence())
+				{
+					const UTakeMetaData* AssetMetaData = LeaderLevelSequence->FindMetaData<UTakeMetaData>();
+
+					UTakeRecorderNamingTokensContext* Context = NewObject<UTakeRecorderNamingTokensContext>();
+					Context->Actor = Actor;
+			
+					AnimationTrackDisplayName = FText::FromString(AssetMetaData->GenerateAssetPath(AnimationTrackDisplayName.ToString(), Context));
+				}
+				
+				AnimTrack->SetDisplayName(AnimationTrackDisplayName);
 				AnimTrack->AddNewAnimation(FFrameNumber(0), AnimSequence.Get());
 				MovieSceneSection = Cast<UMovieSceneSkeletalAnimationSection>(AnimTrack->GetAllSections()[0]);
 				MovieSceneSection->Params.bForceCustomMode = true;
diff --git a/Engine/Plugins/VirtualProduction/Takes/Source/TakeTrackRecorders/Public/TrackRecorders/MovieSceneAnimationTrackRecorderSettings.h b/Engine/Plugins/VirtualProduction/Takes/Source/TakeTrackRecorders/Public/TrackRecorders/MovieSceneAnimationTrackRecorderSettings.h
--- a/Engine/Plugins/VirtualProduction/Takes/Source/TakeTrackRecorders/Public/TrackRecorders/MovieSceneAnimationTrackRecorderSettings.h
+++ b/Engine/Plugins/VirtualProduction/Takes/Source/TakeTrackRecorders/Public/TrackRecorders/MovieSceneAnimationTrackRecorderSettings.h
@@ -27,7 +27,7 @@ public:
 	virtual void PostEditChangeProperty(FPropertyChangedEvent& PropertyChangedEvent) override;
 
 	/** Name of the recorded animation track. */
-	UPROPERTY(config, EditAnywhere, BlueprintReadWrite, Category = "Animation Recorder Settings")
+	UPROPERTY(config, EditAnywhere, BlueprintReadWrite, Category = "Animation Recorder Settings", meta = (NamingTokens))
 	FText AnimationTrackName;
 
 	/** The name of the animation asset. 
@@ -42,7 +42,7 @@ public:
 	 * {slate}     - The slate string.
 	 * {actor}     - The name of the actor being recorded.
 	 */
-	UPROPERTY(config, EditAnywhere, BlueprintReadWrite, Category = "Animation Recorder Settings")
+	UPROPERTY(config, EditAnywhere, BlueprintReadWrite, Category = "Animation Recorder Settings", meta = (NamingTokens))
 	FString AnimationAssetName;
 
 	/** The name of the subdirectory animations will be placed in. Leave this empty to place into the same directory as the sequence base path. 
@@ -57,7 +57,7 @@ public:
 	 * {slate}     - The slate string.
 	 * {actor}     - The name of the actor being recorded.
 	 */
-	UPROPERTY(config, EditAnywhere, BlueprintReadWrite, Category = "Animation Recorder Settings")
+	UPROPERTY(config, EditAnywhere, BlueprintReadWrite, Category = "Animation Recorder Settings", meta = (NamingTokens))
 	FString AnimationSubDirectory;
 	
 	/** Interpolation mode for the recorded keys. */
diff --git a/Engine/Plugins/VirtualProduction/Takes/Source/TakesCore/Private/TakeMetaData.cpp b/Engine/Plugins/VirtualProduction/Takes/Source/TakesCore/Private/TakeMetaData.cpp
--- a/Engine/Plugins/VirtualProduction/Takes/Source/TakesCore/Private/TakeMetaData.cpp
+++ b/Engine/Plugins/VirtualProduction/Takes/Source/TakesCore/Private/TakeMetaData.cpp
@@ -6,6 +6,7 @@
 #include "Engine/Level.h"
 #include "ITakeRecorderNamingTokensModule.h"
 #include "NamingTokenData.h"
+#include "NamingTokens/TakeRecorderNamingTokensContext.h"
 #include "NamingTokensEditorSubsystem.h"
 #include "TakePreset.h"
 #include "TakesCoreBlueprintLibrary.h"
@@ -82,22 +83,8 @@ namespace UE::TakeRecorder::MetaDataHelpers::Private
 			UTakeMetaData* MetaDataCopy = GetCurrentWeakMetaData().IsValid() ?
 				DuplicateObject(GetCurrentWeakMetaData().Get(), GetTransientPackage()) : nullptr;
 			GetLastMetaData().Reset(MetaDataCopy);
-			GetCurrentWeakMetaData().Reset();
 		}
 	}
-
-	/** Update the active metadata for this scope. */
-	struct FCacheMetaDataOnScope
-	{
-		explicit FCacheMetaDataOnScope(UTakeMetaData* InMetaData)
-		{
-			UpdateCachedMetaData(InMetaData);
-		}
-		~FCacheMetaDataOnScope()
-		{
-			UpdateCachedMetaData(nullptr);
-		}
-	};
 }
 
 UTakeMetaData* UTakeMetaData::GetMostRecentMetaData()
@@ -137,22 +124,30 @@ void UTakeMetaData::Unlock()
 	bIsLocked = false;
 }
 
-FString UTakeMetaData::GenerateAssetPath(const FString& PathFormatString) const
+FString UTakeMetaData::GenerateAssetPath(const FString& PathFormatString, UTakeRecorderNamingTokensContext* InContext) const
+{
+	const FNamingTokenResultData Result = ProcessTokens(FText::FromString(PathFormatString), InContext);
+	return Result.EvaluatedText.ToString();
+}
+
+FNamingTokenResultData UTakeMetaData::ProcessTokens(const FText& InText, UTakeRecorderNamingTokensContext* InContext) const
 {
 	if (!ensure(GIsEditor))
 	{
 		UE_LOG(LogTakesCore, Warning, TEXT("GIsEditor is false during UTakeMetaData::GenerateAssetPath and the generated path will be incorrect."))
-		return FString();
+		return FNamingTokenResultData();
 	}
 
-	UE::TakeRecorder::MetaDataHelpers::Private::FCacheMetaDataOnScope ScopedActiveMetaData(const_cast<UTakeMetaData*>(this));
-	
+	if (!InContext)
+	{
+		InContext = NewObject<UTakeRecorderNamingTokensContext>();
+	}
+	InContext->TakeMetaData = this;
+
 	FNamingTokenFilterArgs NamingTokenFilters;
 	NamingTokenFilters.NamespacesToInclude = { ITakeRecorderNamingTokensModule::GetTakeRecorderNamespace() };
-	const FNamingTokenResultData Result = GEditor->GetEditorSubsystem<UNamingTokensEditorSubsystem>()->EvaluateTokenString(PathFormatString,
-		NamingTokenFilters);
-
-	return Result.EvaluatedText.ToString();
+	return GEditor->GetEditorSubsystem<UNamingTokensEditorSubsystem>()->EvaluateTokenText(InText,
+		NamingTokenFilters, { InContext });
 }
 
 const FString& UTakeMetaData::GetSlate() const
diff --git a/Engine/Plugins/VirtualProduction/Takes/Source/TakesCore/Public/NamingTokens/TakeRecorderNamingTokensContext.h b/Engine/Plugins/VirtualProduction/Takes/Source/TakesCore/Public/NamingTokens/TakeRecorderNamingTokensContext.h
new file mode 100644
--- /dev/null
+++ b/Engine/Plugins/VirtualProduction/Takes/Source/TakesCore/Public/NamingTokens/TakeRecorderNamingTokensContext.h
@@ -0,0 +1,35 @@
+// Copyright Epic Games, Inc. All Rights Reserved.
+
+#pragma once
+
+#include "CoreMinimal.h"
+
+#include "TakeRecorderNamingTokensContext.generated.h"
+
+class AActor;
+class UTakeMetaData;
+
+/**
+ * Context object which may be passed to NamingTokens evaluations from within TakeRecorder.
+ * This is stored under TakesCore rather than the TakeRecorderNamingTokensModule for dependency management. Multiple
+ * Take modules need access to the context but don't need access to the NamingTokens themselves and would result
+ * in circular referencing.
+ */
+UCLASS(MinimalAPI)
+class UTakeRecorderNamingTokensContext : public UObject
+{
+	GENERATED_BODY()
+	
+public:
+	/** MetaData specifically for this context. Setting this prevents having to perform a global lookup. */
+	UPROPERTY(Transient)
+	TObjectPtr<const UTakeMetaData> TakeMetaData;
+	
+	/** The specific actor for this context. @todo NamingTokens - Determine if we should be retrieving the {actor} a different way. */
+	UPROPERTY(Transient)
+	TObjectPtr<const AActor> Actor;
+	
+	/** The audio device channel for this context. @todo NamingTokens - Determine if we should be retrieving the {channel} a different way. */
+	UPROPERTY(Transient)
+	int32 AudioInputDeviceChannel = 0;
+};
\ No newline at end of file
diff --git a/Engine/Plugins/VirtualProduction/Takes/Source/TakesCore/Public/TakeMetaData.h b/Engine/Plugins/VirtualProduction/Takes/Source/TakesCore/Public/TakeMetaData.h
--- a/Engine/Plugins/VirtualProduction/Takes/Source/TakesCore/Public/TakeMetaData.h
+++ b/Engine/Plugins/VirtualProduction/Takes/Source/TakesCore/Public/TakeMetaData.h
@@ -10,6 +10,8 @@
 #include "TakeMetaData.generated.h"
 
 class UTakePreset;
+class UTakeRecorderNamingTokensContext;
+struct FNamingTokenResultData;
 
 /**
  * Take meta-data that is stored on ULevelSequence assets that are recorded through the Take Recorder.
@@ -188,8 +190,11 @@ public:
 	 * Generate the desired asset path for this take meta-data
 	 */
 	UFUNCTION(BlueprintCallable, Category="Take")
-	FString GenerateAssetPath(const FString& PathFormatString) const;
-
+	FString GenerateAssetPath(const FString& PathFormatString, UTakeRecorderNamingTokensContext* InContext = nullptr) const;
+	
+	/** Process TakeRecorder Naming Tokens. */
+	FNamingTokenResultData ProcessTokens(const FText& InText, UTakeRecorderNamingTokensContext* InContext = nullptr) const;
+	
 	/**
 	 * Set the slate for this take and reset its take number to 1
 	 * @param bEmitChanged Whether or not to send a slate changed event
