TakeRecorder: Fix custom user tokens not fully supporting case sensitivity.

diff --git a/Engine/Plugins/Developer/NamingTokens/Source/NamingTokens/Private/NamingTokenData.cpp b/Engine/Plugins/Developer/NamingTokens/Source/NamingTokens/Private/NamingTokenData.cpp
new file mode 100644
--- /dev/null
+++ b/Engine/Plugins/Developer/NamingTokens/Source/NamingTokens/Private/NamingTokenData.cpp
@@ -0,0 +1,27 @@
+// Copyright Epic Games, Inc. All Rights Reserved.
+
+#include "NamingTokenData.h"
+
+FNamingTokenData::FNamingTokenData(const FString& InTokenKey)
+{
+	TokenKey = InTokenKey;
+	DisplayName = FText::FromString(InTokenKey);
+}
+
+FNamingTokenData::FNamingTokenData(const FString& InTokenKey, const FText& InTokenDisplayName, const FTokenProcessorDelegateNative& InTokenProcessor)
+{
+	DisplayName = InTokenDisplayName;
+	TokenKey = InTokenKey;
+	TokenProcessorNative = InTokenProcessor;
+}
+
+FNamingTokenData::FNamingTokenData(const FString& InTokenKey, const FText& InTokenDisplayName, const FText& InTokenDescription,
+	const FTokenProcessorDelegateNative& InTokenProcessor): FNamingTokenData(InTokenKey, InTokenDisplayName, InTokenProcessor)
+{
+	Description = InTokenDescription;
+}
+
+bool FNamingTokenData::operator==(const FNamingTokenData& Other) const
+{
+	return TokenKey.Equals(Other.TokenKey, ESearchCase::CaseSensitive);
+}
diff --git a/Engine/Plugins/Developer/NamingTokens/Source/NamingTokens/Public/NamingTokenData.h b/Engine/Plugins/Developer/NamingTokens/Source/NamingTokens/Public/NamingTokenData.h
--- a/Engine/Plugins/Developer/NamingTokens/Source/NamingTokens/Public/NamingTokenData.h
+++ b/Engine/Plugins/Developer/NamingTokens/Source/NamingTokens/Public/NamingTokenData.h
@@ -55,24 +55,10 @@ public:
 
 	DECLARE_DELEGATE_RetVal(FText, FTokenProcessorDelegateNative);
 
-	explicit FNamingTokenData(const FString& InTokenKey)
-	{
-		TokenKey = InTokenKey;
-	}
-	
-	FNamingTokenData(const FString& InTokenKey, const FText& InTokenDisplayName, const FTokenProcessorDelegateNative& InTokenProcessor)
-	{
-		DisplayName = InTokenDisplayName;
-		TokenKey = InTokenKey;
-		TokenProcessorNative = InTokenProcessor;
-	}
-
+	explicit FNamingTokenData(const FString& InTokenKey);
+	FNamingTokenData(const FString& InTokenKey, const FText& InTokenDisplayName, const FTokenProcessorDelegateNative& InTokenProcessor);
 	FNamingTokenData(const FString& InTokenKey, const FText& InTokenDisplayName,
-		const FText& InTokenDescription, const FTokenProcessorDelegateNative& InTokenProcessor)
-		: FNamingTokenData(InTokenKey, InTokenDisplayName, InTokenProcessor)
-	{
-		Description = InTokenDescription;
-	}
+	                 const FText& InTokenDescription, const FTokenProcessorDelegateNative& InTokenProcessor);
 
 	/** The key of the token to use. This is what the text must match in order to be evaluated. Brackets are automatically added and do not need to be included. */
 	UPROPERTY(BlueprintReadWrite, EditAnywhere, Category = "NamingTokens")
@@ -93,11 +79,8 @@ public:
 	/** The native delegate to execute to evaluate the function. If FunctionName is set then this is not used. */
 	FTokenProcessorDelegateNative TokenProcessorNative;
 
-	bool operator==(const FNamingTokenData& Other) const
-	{
-		return TokenKey.Equals(Other.TokenKey, ESearchCase::CaseSensitive);
-	}
-	
+	bool operator==(const FNamingTokenData& Other) const;
+
 	friend FORCEINLINE uint32 GetTypeHash(const FNamingTokenData& Item)
 	{
 		const uint32 Hash = FCrc::StrCrc32(*Item.TokenKey);
diff --git a/Engine/Plugins/VirtualProduction/Takes/Source/TakeRecorder/Private/Widgets/STakeRecorderCockpit.cpp b/Engine/Plugins/VirtualProduction/Takes/Source/TakeRecorder/Private/Widgets/STakeRecorderCockpit.cpp
--- a/Engine/Plugins/VirtualProduction/Takes/Source/TakeRecorder/Private/Widgets/STakeRecorderCockpit.cpp
+++ b/Engine/Plugins/VirtualProduction/Takes/Source/TakeRecorder/Private/Widgets/STakeRecorderCockpit.cpp
@@ -1389,7 +1389,7 @@ void STakeRecorderCockpit::UnbindNamingTokensEvents()
 
 FText STakeRecorderCockpit::GetCustomTokenTextValue(FString InTokenKey) const
 {
-	if (const FText* Value = NamingTokensData->UserDefinedTokens.Find(InTokenKey))
+	if (const FText* Value = NamingTokensData->UserDefinedTokens.Find(FNamingTokenData(InTokenKey)))
 	{
 		return *Value;
 	}
@@ -1400,9 +1400,10 @@ FText STakeRecorderCockpit::GetCustomTokenTextValue(FString InTokenKey) const
 void STakeRecorderCockpit::SetCustomTokenTextValue(const FText& InNewText, ETextCommit::Type InCommitType, FString InTokenKey)
 {
 	FScopedTransaction Transaction(LOCTEXT("SetCustomTokenTextValue_Transaction", "Set Custom Token Text"));
+	const FNamingTokenData CustomToken(InTokenKey);
 	NamingTokensData->Modify();
-	NamingTokensData->UserDefinedTokens.FindOrAdd(InTokenKey) = InNewText;
-	NamingTokensData->VisibleUserTokens.Add(InTokenKey);
+	NamingTokensData->UserDefinedTokens.FindOrAdd(CustomToken) = InNewText;
+	NamingTokensData->VisibleUserTokens.Add(CustomToken.TokenKey);
 	OnTokenValueUpdated();
 }
 
@@ -1426,11 +1427,12 @@ void STakeRecorderCockpit::OnTakeRecorderNamingTokensPreEvaluate(const FNamingTo
 		TArray<FNamingTokenData>& ExternalTokens = NamingTokensData->TakeRecorderNamingTokens->GetExternalTokensChecked(NamingTokensData->NamingTokensExternalGuid);
 
 		ExternalTokens.Reset(NamingTokensData->UserDefinedTokens.Num());
-		for (const TTuple<FString, FText>& UserToken : NamingTokensData->UserDefinedTokens)
+		for (const TTuple<FNamingTokenData, FText>& UserToken : NamingTokensData->UserDefinedTokens)
 		{
 			ExternalTokens.Add({
-				UserToken.Key,
-				FText::FromString(UserToken.Key),
+				UserToken.Key.TokenKey,
+				UserToken.Key.DisplayName,
+				// @todo naming tokens - If we support TOptional<FText> that can be used instead of a lambda
 				FNamingTokenData::FTokenProcessorDelegateNative::CreateLambda([Value = UserToken.Value]() {
 					return Value;
 				})
@@ -1649,7 +1651,7 @@ void STakeRecorderCockpit::EvaluateTokensFromProperty(const FProperty* InPropert
 		for (const FNamingTokenValueData& Token : NamingTokenResultData.TokenValues)
 		{
 			if ((!Token.bWasEvaluated
-				|| NamingTokensData->UserDefinedTokens.Contains(Token.TokenKey))
+				|| NamingTokensData->UserDefinedTokens.Contains(FNamingTokenData(Token.TokenKey)))
 					// Include pre-defined take and slate keys so they will get sorted with user tokens.
 				|| (Token.TokenKey == UE::TakeRecorder::Private::TokenKeyTake
 					|| Token.TokenKey == UE::TakeRecorder::Private::TokenKeySlate))
@@ -1681,30 +1683,31 @@ void STakeRecorderCockpit::CreateUserTokensUI()
 	// Assemble the custom tokens in use. This step needs to happen after NamingTokensData->FieldToUndefinedKeys has
 	// been fully populated so we don't remove custom tokens that are still in use.
 	{
-		const TMap<FString, FText> CustomUserTokensCopy = NamingTokensData->UserDefinedTokens;
-		for (const TTuple<FString, FText>& CustomToken : CustomUserTokensCopy)
+		const TMap<FNamingTokenData, FText> CustomUserTokensCopy = NamingTokensData->UserDefinedTokens;
+		for (const TTuple<FNamingTokenData, FText>& CustomToken : CustomUserTokensCopy)
 		{
 			// If no field holds a reference to this key then we can remove it.
-			if (!NamingTokensData->IsTokenKeyUndefined(CustomToken.Key))
+			if (!NamingTokensData->IsTokenKeyUndefined(CustomToken.Key.TokenKey))
 			{
 				// Only remove the visible entry, but leave the value in CustomUserTokens. If the user is switching
 				// between objects it's possible one object has fields with different tokens than another, and we don't
 				// want to lose the old entered data.
-				NamingTokensData->VisibleUserTokens.Remove(CustomToken.Key);
+				NamingTokensData->VisibleUserTokens.Remove(CustomToken.Key.TokenKey);
 			}
 		}
 
 		// Look for new unidentified tokens that a user can define.
 		for (const FTakeRecorderNamingTokensFieldMapping& UndefinedTokens : NamingTokensData->FieldToUndefinedKeys)
 		{
-			for (const FString& UndefinedToken : UndefinedTokens.UndefinedKeys)
+			for (const FString& UndefinedTokenStr : UndefinedTokens.UndefinedKeys)
 			{
+				FNamingTokenData UndefinedToken(UndefinedTokenStr);
 				if (!NamingTokensData->UserDefinedTokens.Contains(UndefinedToken))
 				{
 					// Check for existence so we don't overwrite the user defined value.
 					NamingTokensData->UserDefinedTokens.Add(UndefinedToken);
 				}
-				NamingTokensData->VisibleUserTokens.Add(UndefinedToken);
+				NamingTokensData->VisibleUserTokens.Add(UndefinedToken.TokenKey);
 			}
 		}
 	}
@@ -1712,9 +1715,7 @@ void STakeRecorderCockpit::CreateUserTokensUI()
 	// Make sure our take dir token example is up-to-date.
 	EvaluateTakeSaveDirTokens();
 
-	// @note User defined tokens are case-insensitive right now. We need to add FLocKeySetFuncs to make it case-sensitive, but because
-	// we only display the token key in UPPERCASE for take recorder, we shouldn't allow the same key with different cases.
-	TSet<FString> UsedTokens;
+	TSet<FString, FLocKeySetFuncs> UsedTokens;
 	UserTokensBox->ClearChildren();
 	
 	// Display tokens in order from their field appearance, to the token location in the string, but do not duplicate.
@@ -1735,7 +1736,7 @@ void STakeRecorderCockpit::CreateUserTokensUI()
 			constexpr float MaxWidth = 175.f;
 
 			// Take and Slate have special handling but still need to show up in the order the user referenced them.
-			if (TokenString == UE::TakeRecorder::Private::TokenKeySlate)
+			if (TokenString.Equals(UE::TakeRecorder::Private::TokenKeySlate))
 			{
 				UserTokensBox->AddSlot()
 				.HAlign(HAlign_Fill)
@@ -1752,7 +1753,7 @@ void STakeRecorderCockpit::CreateUserTokensUI()
 					[
 						SNew(STextBlock)
 						.TextStyle(FTakeRecorderStyle::Get(), "TakeRecorder.TextBox")
-						.Text(LOCTEXT("SlateLabel", "SLATE"))
+						.Text(LOCTEXT("SlateLabel", "slate"))
 					]
 
 					+ SOverlay::Slot()
@@ -1771,7 +1772,7 @@ void STakeRecorderCockpit::CreateUserTokensUI()
 					]
 				];
 			}
-			else if (TokenString == UE::TakeRecorder::Private::TokenKeyTake)
+			else if (TokenString.Equals(UE::TakeRecorder::Private::TokenKeyTake))
 			{
 				UserTokensBox->AddSlot()
 				.HAlign(HAlign_Fill)
@@ -1792,7 +1793,7 @@ void STakeRecorderCockpit::CreateUserTokensUI()
 						[
 							SNew(STextBlock)
 							.TextStyle(FTakeRecorderStyle::Get(), "TakeRecorder.TextBox")
-							.Text(LOCTEXT("TakeLabel", "TAKE"))
+							.Text(LOCTEXT("TakeLabel", "take"))
 						]
 
 						+ SHorizontalBox::Slot()
@@ -1862,7 +1863,7 @@ void STakeRecorderCockpit::CreateUserTokensUI()
 							[
 								SNew(STextBlock)
 								.TextStyle(FTakeRecorderStyle::Get(), "TakeRecorder.TextBox")
-								.Text(FText::FromString(TokenString.ToUpper()))
+								.Text(FText::FromString(TokenString))
 							]
 						]
 						+ SOverlay::Slot()
diff --git a/Engine/Plugins/VirtualProduction/Takes/Source/TakeRecorder/Private/Widgets/TakeRecorderNamingTokensWidgetData.h b/Engine/Plugins/VirtualProduction/Takes/Source/TakeRecorder/Private/Widgets/TakeRecorderNamingTokensWidgetData.h
--- a/Engine/Plugins/VirtualProduction/Takes/Source/TakeRecorder/Private/Widgets/TakeRecorderNamingTokensWidgetData.h
+++ b/Engine/Plugins/VirtualProduction/Takes/Source/TakeRecorder/Private/Widgets/TakeRecorderNamingTokensWidgetData.h
@@ -2,7 +2,8 @@
 
 #pragma once
 
-#include "CoreMinimal.h"
+#include "Internationalization/LocKeyFuncs.h"
+#include "NamingTokenData.h"
 
 #include "TakeRecorderNamingTokensWidgetData.generated.h"
 
@@ -33,16 +34,18 @@ public:
 	/** Guid to our managed tokens. */
 	FGuid NamingTokensExternalGuid;
 	
-	/** Custom tokens entered in by the user, mapped to a user defined value. */
+	/**
+	 * Custom tokens entered in by the user, mapped to a user defined value.
+	 * We use FNamingTokenData rather than an FString for the key so we can support case sensitivity in our map.
+	 */
 	UPROPERTY()
-	TMap<FString, FText> UserDefinedTokens;
+	TMap<FNamingTokenData, FText> UserDefinedTokens;
 
 	/**
 	 * User tokens that are currently visible. Kept as a separate property from UsedDefinedTokens we can persist user values
 	 * between selected sources that have different token entries in their fields.
 	 */
-	UPROPERTY(Transient)
-	TSet<FString> VisibleUserTokens;
+	TSet<FString, FLocKeySetFuncs> VisibleUserTokens;
 	
 	/** Ordered array of FieldName to undefined token keys. These are present if a token evaluation failed to identify them. */
 	UPROPERTY(NonTransactional)
