[TakeRecorder] Support for custom LevelSequence classes

diff --git a/Engine/Plugins/VirtualProduction/Takes/Source/TakeRecorder/Private/Recorder/TakeRecorder.cpp b/Engine/Plugins/VirtualProduction/Takes/Source/TakeRecorder/Private/Recorder/TakeRecorder.cpp
--- a/Engine/Plugins/VirtualProduction/Takes/Source/TakeRecorder/Private/Recorder/TakeRecorder.cpp
+++ b/Engine/Plugins/VirtualProduction/Takes/Source/TakeRecorder/Private/Recorder/TakeRecorder.cpp
@@ -1177,7 +1177,7 @@ void UTakeRecorder::StopInternal(const bool bCancelled)
 			// Rebuild sequencer because subsequences could have been added or bindings removed
 			if (Sequencer)
 			{
-				Sequencer->RefreshTree();
+				Sequencer->NotifyMovieSceneDataChanged(EMovieSceneDataChangeType::MovieSceneStructureItemsChanged);
 			}
 		}
 	}
diff --git a/Engine/Plugins/VirtualProduction/Takes/Source/TakesCore/Private/TakeRecorderSources.cpp b/Engine/Plugins/VirtualProduction/Takes/Source/TakesCore/Private/TakeRecorderSources.cpp
--- a/Engine/Plugins/VirtualProduction/Takes/Source/TakesCore/Private/TakeRecorderSources.cpp
+++ b/Engine/Plugins/VirtualProduction/Takes/Source/TakesCore/Private/TakeRecorderSources.cpp
@@ -561,7 +561,7 @@ ULevelSequence* UTakeRecorderSources::CreateSubSequenceForSource(ULevelSequence*
 	FString NewPath = FString::Printf(TEXT("%s/%s_Subscenes/%s"), *SequenceDirectory, *SequenceName, *SubSequenceAssetName);
 
 	ULevelSequence* OutAsset = nullptr;
-	TakesUtils::CreateNewAssetPackage<ULevelSequence>(NewPath, OutAsset, nullptr, ExistingSubSequence);
+	TakesUtils::CreateNewAssetPackage<ULevelSequence>(NewPath, OutAsset, nullptr, ExistingSubSequence, InRootSequence->GetClass());
 	if (OutAsset)
 	{
 
diff --git a/Engine/Plugins/VirtualProduction/Takes/Source/TakesCore/Public/TakesUtils.h b/Engine/Plugins/VirtualProduction/Takes/Source/TakesCore/Public/TakesUtils.h
--- a/Engine/Plugins/VirtualProduction/Takes/Source/TakesCore/Public/TakesUtils.h
+++ b/Engine/Plugins/VirtualProduction/Takes/Source/TakesCore/Public/TakesUtils.h
@@ -49,7 +49,7 @@ namespace TakesUtils
 	* @return True if the asset was created successfully, false if there was an error.
 	*/
 	template<typename AssetType>
-	static bool CreateNewAssetPackage(FString &InPackageName, AssetType*& OutAsset, FText* OutError, AssetType* OptionalBase = nullptr)
+	static bool CreateNewAssetPackage(FString &InPackageName, AssetType*& OutAsset, FText* OutError, AssetType* OptionalBase = nullptr, UClass* OptionalClass = nullptr)
 	{
 		if (!FPackageName::IsValidLongPackageName(InPackageName))
 		{
@@ -92,6 +92,10 @@ namespace TakesUtils
 			OutAsset = Cast<AssetType>(StaticDuplicateObject(OptionalBase, NewPackage, *NewAssetName, RF_NoFlags));
 			OutAsset->SetFlags(RF_Public | RF_Standalone | RF_Transactional);
 		}
+		else if (OptionalClass)
+		{
+			OutAsset = NewObject<AssetType>(NewPackage, OptionalClass, *NewAssetName, RF_Public | RF_Standalone | RF_Transactional);
+		}
 		else
 		{
 			// Create a new level sequence from scratch
@@ -102,10 +106,10 @@ namespace TakesUtils
 	}
 
 	template<typename AssetType>
-	static bool CreateNewAssetPackage(FString& InPackageName, TObjectPtr<AssetType>& OutAsset, FText* OutError, AssetType* OptionalBase = nullptr)
+	static bool CreateNewAssetPackage(FString& InPackageName, TObjectPtr<AssetType>& OutAsset, FText* OutError, AssetType* OptionalBase = nullptr, UClass* OptionalClass = nullptr)
 	{
 		AssetType* RawOutAsset = nullptr;
-		if (CreateNewAssetPackage(InPackageName, RawOutAsset, OutError, OptionalBase))
+		if (CreateNewAssetPackage(InPackageName, RawOutAsset, OutError, OptionalBase, OptionalClass))
 		{
 			OutAsset = RawOutAsset;
 			return true;
