[Take Recorder/PCap Workflow]

Added BlueprintReadOnly access to the RecordedObject property.

This is so users can query the class of the recorded object and filter whether or not to record it's properties.

diff --git a/Engine/Plugins/VirtualProduction/Takes/Source/TakesCore/Public/TakeRecorderSourceProperty.h b/Engine/Plugins/VirtualProduction/Takes/Source/TakesCore/Public/TakeRecorderSourceProperty.h
--- a/Engine/Plugins/VirtualProduction/Takes/Source/TakesCore/Public/TakeRecorderSourceProperty.h
+++ b/Engine/Plugins/VirtualProduction/Takes/Source/TakesCore/Public/TakeRecorderSourceProperty.h
@@ -50,7 +50,7 @@ public:
 	virtual void PostEditUndo() override;
 	virtual void PostEditChangeProperty(struct FPropertyChangedEvent& PropertyChangedEvent) override;
 
-	UPROPERTY(VisibleAnywhere, Category = "Property")
+	UPROPERTY(BlueprintReadOnly, VisibleAnywhere, Category = "Property")
 	TSoftObjectPtr<UObject> RecordedObject;
 
 	/* Represents properties exposed to Cinematics that can possibly be recorded. */
