Take Recorder: Extract all customizations to dedicated files and put them all into the Customization folder.

diff --git a/Engine/Plugins/VirtualProduction/Takes/Source/TakeRecorder/Private/Customization/AudioInputChannelPropertyCustomization.cpp b/Engine/Plugins/VirtualProduction/Takes/Source/TakeRecorder/Private/Customization/AudioInputChannelPropertyCustomization.cpp
new file mode 100644
--- /dev/null
+++ b/Engine/Plugins/VirtualProduction/Takes/Source/TakeRecorder/Private/Customization/AudioInputChannelPropertyCustomization.cpp
@@ -0,0 +1,125 @@
+﻿// Copyright Epic Games, Inc. All Rights Reserved.
+
+#include "AudioInputChannelPropertyCustomization.h"
+
+#include "IDetailChildrenBuilder.h"
+#include "IDetailPropertyRow.h"
+#include "TakeRecorderAudioSettingsCustomization.h"
+#include "TakeRecorderSourceProperty.h"
+#include "Widgets/Input/SComboBox.h"
+#include "Widgets/Text/STextBlock.h"
+
+namespace UE::TakeRecorder
+{
+void FAudioInputChannelPropertyCustomization::CustomizeChildren(
+	TSharedRef<IPropertyHandle> PropertyHandle, IDetailChildrenBuilder& ChildBuilder, IPropertyTypeCustomizationUtils& CustomizationUtils
+	)
+{
+	InputDeviceChannelHandle = PropertyHandle->GetChildHandle(GET_MEMBER_NAME_CHECKED(FAudioInputDeviceChannelProperty, AudioInputDeviceChannel));
+
+	if (InputDeviceChannelHandle.IsValid())
+	{
+		IDetailPropertyRow& InputDeviceChannelPropertyRow = ChildBuilder.AddProperty(PropertyHandle);
+		InputDeviceChannelPropertyRow.CustomWidget()
+		.NameContent()
+		[
+			InputDeviceChannelHandle->CreatePropertyNameWidget()
+		]
+		.ValueContent()
+		[
+			MakeInputChannelSelectorWidget()
+		];
+	}
+
+	// Register delegate so UI can update when the audio input device changes
+	if (UTakeRecorderAudioInputSettings* AudioInputSettings = TakeRecorderAudioSettingsUtils::GetTakeRecorderAudioInputSettings())
+	{
+		AudioInputSettings->GetOnAudioInputDeviceChanged().Add(FSimpleDelegate::CreateSP(this, &FAudioInputChannelPropertyCustomization::RebuildInputChannelArray));
+	}
+}
+
+void FAudioInputChannelPropertyCustomization::BuildInputChannelArray()
+{
+	if (UTakeRecorderAudioInputSettings* AudioInputSettings = TakeRecorderAudioSettingsUtils::GetTakeRecorderAudioInputSettings())
+	{
+		int32 DeviceChannelCount = AudioInputSettings->GetDeviceChannelCount();
+
+		if (DeviceChannelCount > 0)
+		{
+			for (int32 ChannelIndex = 1; ChannelIndex <= DeviceChannelCount; ++ChannelIndex)
+			{
+				TSharedPtr<int32> ChannelIndexPtr = MakeShared<int32>(ChannelIndex);
+				InputDeviceChannelArray.Add(ChannelIndexPtr);
+			}
+		}
+
+		int32 CurrentValue;
+		InputDeviceChannelHandle->GetValue(CurrentValue);
+		if (CurrentValue > DeviceChannelCount)
+		{
+			InputDeviceChannelHandle->SetValue(0);
+		}
+	}
+}
+
+void FAudioInputChannelPropertyCustomization::RebuildInputChannelArray()
+{
+	InputDeviceChannelArray.Empty();
+	BuildInputChannelArray();
+	ChannelComboBox->RefreshOptions();
+
+	UTakeRecorderAudioInputSettings* AudioInputSettings = TakeRecorderAudioSettingsUtils::GetTakeRecorderAudioInputSettings();
+
+	if (AudioInputSettings && InputDeviceChannelHandle.IsValid())
+	{
+		int32 CurrentValue;
+		InputDeviceChannelHandle->GetValue(CurrentValue);
+		int32 DeviceChannelCount = AudioInputSettings->GetDeviceChannelCount();
+
+		if (CurrentValue < 1 || CurrentValue > DeviceChannelCount)
+		{
+			InputChannelTitleBlock->SetText(FText());
+		}
+	}
+}
+
+TSharedRef<SWidget> FAudioInputChannelPropertyCustomization::MakeInputChannelSelectorWidget()
+{
+	BuildInputChannelArray();
+
+	// Combo box component:
+	SAssignNew(ChannelComboBox, SComboBox<TSharedPtr<int32>>)
+		.OptionsSource(&InputDeviceChannelArray)
+		.OnGenerateWidget_Lambda([](TSharedPtr<int32> InValue)
+		{
+			return SNew(STextBlock)
+				.Text(FText::AsNumber(*InValue))
+				.Font(FAppStyle::Get().GetFontStyle("SmallFont"));
+		})
+		.OnSelectionChanged_Lambda([this](TSharedPtr<int32> InSelection, ESelectInfo::Type InSelectInfo)
+		{
+			if (InSelection.IsValid())
+			{
+				InputDeviceChannelHandle->SetValue(*InSelection);
+				InputChannelTitleBlock->SetText(FText::AsNumber(*InSelection));
+			}
+		})
+		[
+			SAssignNew(InputChannelTitleBlock, STextBlock)
+			.Text_Lambda([this]()
+			{
+				int32 ChannelNumber;
+				InputDeviceChannelHandle->GetValue(ChannelNumber);
+
+				if (ChannelNumber > 0)
+				{
+					return FText::AsNumber(ChannelNumber);
+				}
+				return FText();
+			})
+			.Font(FAppStyle::Get().GetFontStyle("SmallFont"))
+		];
+
+	return ChannelComboBox.ToSharedRef();
+}
+}
diff --git a/Engine/Plugins/VirtualProduction/Takes/Source/TakeRecorder/Private/Customization/AudioInputChannelPropertyCustomization.h b/Engine/Plugins/VirtualProduction/Takes/Source/TakeRecorder/Private/Customization/AudioInputChannelPropertyCustomization.h
new file mode 100644
--- /dev/null
+++ b/Engine/Plugins/VirtualProduction/Takes/Source/TakeRecorder/Private/Customization/AudioInputChannelPropertyCustomization.h
@@ -0,0 +1,38 @@
+﻿// Copyright Epic Games, Inc. All Rights Reserved.
+
+#pragma once
+
+#include "Containers/Array.h"
+#include "HAL/Platform.h"
+#include "IPropertyTypeCustomization.h"
+#include "Templates/SharedPointer.h"
+
+class IPropertyHandle;
+class SWidget;
+class STextBlock;
+template<typename OptionType> class SComboBox;
+
+namespace UE::TakeRecorder
+{
+class FAudioInputChannelPropertyCustomization : public IPropertyTypeCustomization
+{
+public:
+
+	//~ Begin IPropertyTypeCustomization Interface
+	virtual void CustomizeHeader(TSharedRef<IPropertyHandle> PropertyHandle, FDetailWidgetRow& HeaderRow, IPropertyTypeCustomizationUtils& ) {}
+	virtual void CustomizeChildren(TSharedRef<IPropertyHandle> PropertyHandle, IDetailChildrenBuilder& ChildBuilder, IPropertyTypeCustomizationUtils& CustomizationUtils);
+	//~ End IPropertyTypeCustomization Interface
+
+private:
+
+	TSharedPtr<IPropertyHandle> InputDeviceChannelHandle;
+	TArray<TSharedPtr<int32>> InputDeviceChannelArray;
+	TSharedPtr<STextBlock> InputChannelTitleBlock;
+	TSharedPtr<SComboBox<TSharedPtr<int32>>> ChannelComboBox;
+		
+	void BuildInputChannelArray();
+	void RebuildInputChannelArray();
+
+	TSharedRef<SWidget> MakeInputChannelSelectorWidget();
+};
+}
\ No newline at end of file
diff --git a/Engine/Plugins/VirtualProduction/Takes/Source/TakeRecorder/Private/Customization/RecordedPropertyCustomization.cpp b/Engine/Plugins/VirtualProduction/Takes/Source/TakeRecorder/Private/Customization/RecordedPropertyCustomization.cpp
new file mode 100644
--- /dev/null
+++ b/Engine/Plugins/VirtualProduction/Takes/Source/TakeRecorder/Private/Customization/RecordedPropertyCustomization.cpp
@@ -0,0 +1,54 @@
+﻿// Copyright Epic Games, Inc. All Rights Reserved.
+
+#include "RecordedPropertyCustomization.h"
+
+#include "DetailWidgetRow.h"
+#include "PropertyHandle.h"
+#include "TakeRecorderSourceProperty.h"
+
+namespace UE::TakeRecorder
+{
+const FString FRecordedPropertyCustomization::PropertyPathDelimiter = FString(TEXT("."));
+	
+void FRecordedPropertyCustomization::CustomizeHeader(
+	TSharedRef<IPropertyHandle> PropertyHandle, FDetailWidgetRow& HeaderRow, IPropertyTypeCustomizationUtils& CustomizationUtils
+	)
+{
+	if (PropertyHandle->IsValidHandle())
+	{
+		TSharedPtr<IPropertyHandle> PropertyNameHandle = PropertyHandle->GetChildHandle(GET_MEMBER_NAME_CHECKED(FActorRecordedProperty, PropertyName)); 
+		TSharedPtr<IPropertyHandle> bEnabledHandle = PropertyHandle->GetChildHandle(GET_MEMBER_NAME_CHECKED(FActorRecordedProperty, bEnabled)); 
+		TSharedPtr<IPropertyHandle> RecorderNameHandle = PropertyHandle->GetChildHandle(GET_MEMBER_NAME_CHECKED(FActorRecordedProperty, RecorderName));
+
+		FString PropertyNameValue;
+		PropertyNameHandle->GetValueAsDisplayString(PropertyNameValue);
+
+		FString ParentGroups; 
+		FString PropertyName;
+		FText DisplayString = FText::FromString(PropertyNameValue);
+		if (PropertyNameValue.Split(PropertyPathDelimiter, &ParentGroups, &PropertyName, ESearchCase::IgnoreCase, ESearchDir::FromEnd))
+		{
+			DisplayString = FText::FromString(PropertyName);
+		}
+
+		HeaderRow 
+		[
+			SNew(SHorizontalBox)
+
+			+SHorizontalBox::Slot()
+			.AutoWidth()
+			[
+				bEnabledHandle->CreatePropertyValueWidget(false)
+			]
+
+			+SHorizontalBox::Slot()
+			.Padding(8, 0, 0, 0)
+			[
+				PropertyNameHandle->CreatePropertyNameWidget(DisplayString)
+			]
+
+		];
+	}
+}
+}
+
diff --git a/Engine/Plugins/VirtualProduction/Takes/Source/TakeRecorder/Private/Customization/RecordedPropertyCustomization.h b/Engine/Plugins/VirtualProduction/Takes/Source/TakeRecorder/Private/Customization/RecordedPropertyCustomization.h
new file mode 100644
--- /dev/null
+++ b/Engine/Plugins/VirtualProduction/Takes/Source/TakeRecorder/Private/Customization/RecordedPropertyCustomization.h
@@ -0,0 +1,22 @@
+﻿// Copyright Epic Games, Inc. All Rights Reserved.
+
+#pragma once
+
+#include "IPropertyTypeCustomization.h"
+#include "Containers/UnrealString.h"
+
+namespace UE::TakeRecorder
+{
+class FRecordedPropertyCustomization : public IPropertyTypeCustomization
+{
+	static const FString PropertyPathDelimiter;
+
+	//~ Begin IPropertyTypeCustomization Interface
+	virtual void CustomizeHeader(TSharedRef<IPropertyHandle> PropertyHandle, FDetailWidgetRow& HeaderRow, IPropertyTypeCustomizationUtils& CustomizationUtils );
+	virtual void CustomizeChildren(TSharedRef<IPropertyHandle> PropertyHandle, IDetailChildrenBuilder& ChildBuilder, IPropertyTypeCustomizationUtils& CustomizationUtils )
+	{
+		// Intentionally Left Blank, Child Customization was handled in the Header Row
+	}
+	//~ End IPropertyTypeCustomization Interface
+};
+}
\ No newline at end of file
diff --git a/Engine/Plugins/VirtualProduction/Takes/Source/TakeRecorder/Private/Customization/RecorderPropertyMapCustomization.cpp b/Engine/Plugins/VirtualProduction/Takes/Source/TakeRecorder/Private/Customization/RecorderPropertyMapCustomization.cpp
new file mode 100644
--- /dev/null
+++ b/Engine/Plugins/VirtualProduction/Takes/Source/TakeRecorder/Private/Customization/RecorderPropertyMapCustomization.cpp
@@ -0,0 +1,383 @@
+﻿// Copyright Epic Games, Inc. All Rights Reserved.
+
+#include "RecorderPropertyMapCustomization.h"
+
+#include "ClassIconFinder.h"
+#include "DetailWidgetRow.h"
+#include "GameFramework/Actor.h"
+#include "IDetailChildrenBuilder.h"
+#include "IDetailGroup.h"
+#include "TakeRecorderSourceProperty.h"
+#include "Widgets/Images/SImage.h"
+#include "Widgets/Input/SCheckBox.h"
+#include "Widgets/Text/STextBlock.h"
+
+#define LOCTEXT_NAMESPACE "SLevelSequenceTakeEditor"
+
+namespace UE::TakeRecorder
+{
+const FString FRecorderPropertyMapCustomization::PropertyPathDelimiter = FString(TEXT("."));
+	
+void FRecorderPropertyMapCustomization::CustomizeHeader(
+	TSharedRef<IPropertyHandle> PropertyHandle, FDetailWidgetRow& HeaderRow, IPropertyTypeCustomizationUtils& CustomizationUtils
+	)
+{
+	TSharedPtr<IPropertyHandle> RecordedObjectHandle = PropertyHandle->GetChildHandle(GET_MEMBER_NAME_CHECKED(UActorRecorderPropertyMap, RecordedObject)); 
+
+	UObject* Object = nullptr;
+	FText ActorOrComponentName = LOCTEXT("MissingActorOrComponentName", "MissingActorOrComponentName");
+
+	const FSlateBrush* Icon = nullptr;
+
+	if ( RecordedObjectHandle->IsValidHandle() && RecordedObjectHandle->GetValue(Object) == FPropertyAccess::Success && Object != nullptr )
+	{
+		if (AActor* Actor = Cast<AActor>(Object))
+		{
+			ActorOrComponentName = FText::AsCultureInvariant(Actor->GetActorLabel());
+			Icon = FClassIconFinder::FindIconForActor(Actor);
+		}
+		else
+		{
+			ActorOrComponentName = FText::AsCultureInvariant(Object->GetName());
+			Icon = FSlateIconFinder::FindIconBrushForClass(Object->GetClass());
+		}
+	}
+
+	HeaderRow
+	[
+		SNew(SHorizontalBox)
+
+		+SHorizontalBox::Slot()
+		.AutoWidth()
+		.VAlign(VAlign_Center)
+		[
+			SNew( SCheckBox )
+			.OnCheckStateChanged( this, &FRecorderPropertyMapCustomization::OnCheckStateChanged, PropertyHandle )	
+			.IsChecked( this, &FRecorderPropertyMapCustomization::OnGetCheckState, PropertyHandle )	
+			.Padding(0.0)
+		]
+
+		+SHorizontalBox::Slot()
+		.AutoWidth()
+		.VAlign(VAlign_Center)
+		.Padding(8, 0, 0, 0)
+		[
+			SNew(SImage)
+			.Image(Icon)
+		]
+
+		+SHorizontalBox::Slot()
+		.VAlign(VAlign_Center)
+		.AutoWidth()
+		.Padding(2.0, 0.0)
+		[
+			SNew(STextBlock)
+			.Text(ActorOrComponentName)
+			.Font( FAppStyle::GetFontStyle( "PropertyWindow.BoldFont" ))
+		]
+
+		+SHorizontalBox::Slot()
+		.FillWidth(1.0f)
+		.VAlign(VAlign_Center)
+		.Padding(2.0, 0.0)
+		[
+			SNew(STextBlock)
+			.Text(LOCTEXT("TakeRecorderRecordedPropertiesTitle", "Recorded Properties"))
+			.Font( FAppStyle::GetFontStyle( "PropertyWindow.NormalFont" ))
+		]
+	];
+}
+
+IDetailGroup& FRecorderPropertyMapCustomization::GetOrCreateDetailGroup(
+	IDetailChildrenBuilder& ChildBuilder, TMap<FString, IDetailGroup*>& GroupMap, TSharedPtr<IPropertyHandleArray> PropertiesArray, FString& GroupName
+	)
+{
+	if (GroupMap.Contains(GroupName))
+	{
+		return *GroupMap[GroupName];
+	}
+
+	FString ParentGroups; 
+	FString PropertyName;
+	FText DisplayName;
+
+	if (GroupName.Split(PropertyPathDelimiter, &ParentGroups, &PropertyName, ESearchCase::IgnoreCase, ESearchDir::FromEnd))
+	{
+		IDetailGroup& ParentGroup = GetOrCreateDetailGroup(ChildBuilder, GroupMap, PropertiesArray, ParentGroups);
+		DisplayName = FText::FromString(PropertyName);
+		GroupMap.Add(GroupName, &ParentGroup.AddGroup(FName(*PropertyName), DisplayName));
+	}
+	else 
+	{
+		DisplayName = FText::FromString(GroupName);
+		GroupMap.Add(GroupName, &ChildBuilder.AddGroup(FName(*GroupName), DisplayName));
+	}
+
+	GroupMap[GroupName]->HeaderRow()
+	[
+		SNew(SHorizontalBox)
+
+		+SHorizontalBox::Slot()
+		.AutoWidth()
+		.VAlign(VAlign_Center)
+		[
+			SNew( SCheckBox )
+			.OnCheckStateChanged( this, &FRecorderPropertyMapCustomization::OnGroupCheckStateChanged, PropertiesArray, GroupName )	
+			.IsChecked( this, &FRecorderPropertyMapCustomization::OnGroupGetCheckState, PropertiesArray, GroupName )	
+		]
+
+		+SHorizontalBox::Slot()
+		.VAlign(VAlign_Center)
+		.AutoWidth()
+		.Padding(6.0, 0.0)
+		[
+			SNew(STextBlock)
+			.Text(DisplayName)
+			.Font( FAppStyle::GetFontStyle( "PropertyWindow.NormalFont" ))
+		]
+	];
+
+	return *GroupMap[GroupName];
+}
+
+void FRecorderPropertyMapCustomization::CustomizeChildren(
+	TSharedRef<IPropertyHandle> PropertyHandle, IDetailChildrenBuilder& ChildBuilder, IPropertyTypeCustomizationUtils& CustomizationUtils
+	)
+{
+	TMap<FString, IDetailGroup*> DetailGroupMap;
+
+	TSharedPtr<IPropertyHandleArray> RecordedPropertiesArrayHandle = PropertyHandle->GetChildHandle(GET_MEMBER_NAME_CHECKED(UActorRecorderPropertyMap, Properties))->AsArray();
+
+	uint32 NumProperties;
+	RecordedPropertiesArrayHandle->GetNumElements(NumProperties);
+
+
+	for(uint32 i = 0; i < NumProperties; ++i)
+	{
+		TSharedRef<IPropertyHandle> RecordedPropertyTemp = RecordedPropertiesArrayHandle->GetElement(i);
+		if (RecordedPropertyTemp->IsValidHandle())
+		{
+
+			TSharedPtr<IPropertyHandle> PropertyNameHandle = RecordedPropertyTemp->GetChildHandle(GET_MEMBER_NAME_CHECKED(FActorRecordedProperty, PropertyName)); 
+			if (PropertyNameHandle->IsValidHandle())
+			{
+				FString PropertyNameValue;
+				PropertyNameHandle->GetValueAsDisplayString(PropertyNameValue);
+
+				FString ParentGroups; 
+				FString PropertyName;
+				if (PropertyNameValue.Split(PropertyPathDelimiter, &ParentGroups, &PropertyName, ESearchCase::IgnoreCase, ESearchDir::FromEnd))
+				{
+					IDetailGroup& ParentGroup = GetOrCreateDetailGroup(ChildBuilder, DetailGroupMap, RecordedPropertiesArrayHandle, ParentGroups);
+					ParentGroup.AddPropertyRow(RecordedPropertyTemp);
+				}
+				else 
+				{
+					ChildBuilder.AddProperty(RecordedPropertyTemp);
+				}
+			}
+		}
+	}
+
+	TSharedPtr<IPropertyHandleArray> RecordedComponentsArrayHandle = PropertyHandle->GetChildHandle(GET_MEMBER_NAME_CHECKED(UActorRecorderPropertyMap, Children))->AsArray();
+
+	uint32 NumCompProperties;
+	RecordedComponentsArrayHandle->GetNumElements(NumCompProperties);
+
+	for(uint32 i = 0; i < NumCompProperties; ++i)
+	{
+		TSharedRef<IPropertyHandle> RecordedCompTemp = RecordedComponentsArrayHandle->GetElement(i);
+		if (RecordedCompTemp->IsValidHandle())
+		{
+			ChildBuilder.AddProperty(RecordedCompTemp);
+		}
+	}
+}
+
+void FRecorderPropertyMapCustomization::OnGroupCheckStateChanged(
+	ECheckBoxState InNewState, TSharedPtr<IPropertyHandleArray> RecordedPropertiesArrayHandle, FString GroupName
+	) const
+{
+	uint32 NumProperties;
+	RecordedPropertiesArrayHandle->GetNumElements(NumProperties);
+
+	for(uint32 i = 0; i < NumProperties; ++i)
+	{
+		TSharedRef<IPropertyHandle> RecordedPropertyTemp = RecordedPropertiesArrayHandle->GetElement(i);
+		if (RecordedPropertyTemp->IsValidHandle())
+		{
+			TSharedPtr<IPropertyHandle> PropertyNameHandle = RecordedPropertyTemp->GetChildHandle(GET_MEMBER_NAME_CHECKED(FActorRecordedProperty, PropertyName)); 
+			if (PropertyNameHandle->IsValidHandle())
+			{
+				FString PropertyNameValue;
+				PropertyNameHandle->GetValueAsDisplayString(PropertyNameValue);
+
+				if (PropertyNameValue.StartsWith(GroupName))
+				{
+					TSharedPtr<IPropertyHandle> EnabledRecordedPropertyTemp = RecordedPropertyTemp->GetChildHandle(GET_MEMBER_NAME_CHECKED(FActorRecordedProperty, bEnabled));
+					if (EnabledRecordedPropertyTemp->IsValidHandle())
+					{
+						EnabledRecordedPropertyTemp->SetValue( InNewState == ECheckBoxState::Checked ? true : false );
+					}
+				}
+			}
+		}
+	}
+}
+
+ECheckBoxState FRecorderPropertyMapCustomization::OnGroupGetCheckState(
+	TSharedPtr<IPropertyHandleArray> RecordedPropertiesArrayHandle, FString GroupName
+	) const
+{
+	bool SetFirst = false;
+	bool FinalCheckedValue = false;
+
+	uint32 NumProperties;
+	RecordedPropertiesArrayHandle->GetNumElements(NumProperties);
+
+	for(uint32 i = 0; i < NumProperties; ++i)
+	{
+		TSharedRef<IPropertyHandle> RecordedPropertyTemp = RecordedPropertiesArrayHandle->GetElement(i);
+		if (RecordedPropertyTemp->IsValidHandle())
+		{
+			TSharedPtr<IPropertyHandle> PropertyNameHandle = RecordedPropertyTemp->GetChildHandle(GET_MEMBER_NAME_CHECKED(FActorRecordedProperty, PropertyName)); 
+			if (PropertyNameHandle->IsValidHandle())
+			{
+				FString PropertyNameValue;
+				PropertyNameHandle->GetValueAsDisplayString(PropertyNameValue);
+
+				if (PropertyNameValue.StartsWith(GroupName))
+				{
+					TSharedPtr<IPropertyHandle> EnabledRecordedPropertyTemp = RecordedPropertyTemp->GetChildHandle(GET_MEMBER_NAME_CHECKED(FActorRecordedProperty, bEnabled));
+					if (EnabledRecordedPropertyTemp->IsValidHandle())
+					{
+						bool EnabledValue;
+						if ( EnabledRecordedPropertyTemp->GetValue(EnabledValue) == FPropertyAccess::Success )
+						{
+							if (!SetFirst)
+							{
+								FinalCheckedValue = EnabledValue;
+								SetFirst = true;
+							}
+							else if ( EnabledValue != FinalCheckedValue )
+							{
+								return ECheckBoxState::Undetermined;
+							}
+						}
+					}
+				}
+			}
+		}
+	}
+
+	return FinalCheckedValue ? ECheckBoxState::Checked : ECheckBoxState::Unchecked;
+}
+
+void FRecorderPropertyMapCustomization::OnCheckStateChanged(ECheckBoxState InNewState, TSharedRef<IPropertyHandle> PropertyHandle)
+{
+	TSharedPtr<IPropertyHandleArray> RecordedPropertiesArrayHandle = PropertyHandle->GetChildHandle(GET_MEMBER_NAME_CHECKED(UActorRecorderPropertyMap, Properties))->AsArray();
+
+	uint32 NumProperties;
+	RecordedPropertiesArrayHandle->GetNumElements(NumProperties);
+
+	for(uint32 i = 0; i < NumProperties; ++i)
+	{
+		TSharedRef<IPropertyHandle> RecordedPropertyTemp = RecordedPropertiesArrayHandle->GetElement(i);
+		if (RecordedPropertyTemp->IsValidHandle())
+		{
+			TSharedPtr<IPropertyHandle> EnabledRecordedPropertyTemp = RecordedPropertyTemp->GetChildHandle(GET_MEMBER_NAME_CHECKED(FActorRecordedProperty, bEnabled));
+			if (EnabledRecordedPropertyTemp->IsValidHandle())
+			{
+				EnabledRecordedPropertyTemp->SetValue( InNewState == ECheckBoxState::Checked );
+			}
+		}
+	}
+
+	TSharedPtr<IPropertyHandleArray> RecordedComponentsArrayHandle = PropertyHandle->GetChildHandle(GET_MEMBER_NAME_CHECKED(UActorRecorderPropertyMap, Children))->AsArray();
+
+	uint32 NumCompProperties;
+	RecordedComponentsArrayHandle->GetNumElements(NumCompProperties);
+
+	for(uint32 i = 0; i < NumCompProperties; ++i)
+	{
+		TSharedPtr<IPropertyHandle> RecordedCompTemp = RecordedComponentsArrayHandle->GetElement(i);
+		if (RecordedCompTemp->IsValidHandle())
+		{
+			OnCheckStateChanged( InNewState, RecordedCompTemp.ToSharedRef());
+		}
+	}
+}
+
+ECheckBoxState FRecorderPropertyMapCustomization::OnGetCheckState(TSharedRef<IPropertyHandle> PropertyHandle) const
+{
+	bool SetFirst = false;
+		bool FinalCheckedValue = false;
+
+		TSharedPtr<IPropertyHandleArray> RecordedPropertiesArrayHandle = PropertyHandle->GetChildHandle(GET_MEMBER_NAME_CHECKED(UActorRecorderPropertyMap, Properties))->AsArray();
+
+		uint32 NumProperties;
+		RecordedPropertiesArrayHandle->GetNumElements(NumProperties);
+
+		for(uint32 i = 0; i < NumProperties; ++i)
+		{
+			TSharedRef<IPropertyHandle> RecordedPropertyTemp = RecordedPropertiesArrayHandle->GetElement(i);
+			if (RecordedPropertyTemp->IsValidHandle())
+			{
+
+				TSharedPtr<IPropertyHandle> EnabledRecordedPropertyTemp = RecordedPropertyTemp->GetChildHandle(GET_MEMBER_NAME_CHECKED(FActorRecordedProperty, bEnabled));
+				if (EnabledRecordedPropertyTemp->IsValidHandle())
+				{
+					bool EnabledValue;
+					if ( EnabledRecordedPropertyTemp->GetValue(EnabledValue) == FPropertyAccess::Success )
+					{
+						if (!SetFirst)
+						{
+							FinalCheckedValue = EnabledValue;
+							SetFirst = true;
+						}
+						else if ( EnabledValue != FinalCheckedValue )
+						{
+							return ECheckBoxState::Undetermined;
+						}
+					}
+				}
+			}
+		}
+
+		TSharedPtr<IPropertyHandleArray> RecordedComponentsArrayHandle = PropertyHandle->GetChildHandle(GET_MEMBER_NAME_CHECKED(UActorRecorderPropertyMap, Children))->AsArray();
+
+		uint32 NumCompProperties;
+		RecordedComponentsArrayHandle->GetNumElements(NumCompProperties);
+
+		for(uint32 i = 0; i < NumCompProperties; ++i)
+		{
+			TSharedPtr<IPropertyHandle> RecordedCompTemp = RecordedComponentsArrayHandle->GetElement(i);
+			if (RecordedCompTemp->IsValidHandle())
+			{
+				ECheckBoxState CompEnabledState = OnGetCheckState( RecordedCompTemp.ToSharedRef() );
+				if (CompEnabledState == ECheckBoxState::Undetermined)
+				{
+					return ECheckBoxState::Undetermined;
+				}
+				else 
+				{
+					bool IsCompChecked = CompEnabledState == ECheckBoxState::Checked ? true : false;
+
+					if ( !SetFirst )
+					{
+						FinalCheckedValue = IsCompChecked;	
+						SetFirst = true;
+					}
+					else if ( IsCompChecked != FinalCheckedValue )
+					{
+						return ECheckBoxState::Undetermined;
+					}
+				}
+			}
+		}
+
+		return FinalCheckedValue ? ECheckBoxState::Checked : ECheckBoxState::Unchecked;
+}
+}
+
+#undef LOCTEXT_NAMESPACE
\ No newline at end of file
diff --git a/Engine/Plugins/VirtualProduction/Takes/Source/TakeRecorder/Private/Customization/RecorderPropertyMapCustomization.h b/Engine/Plugins/VirtualProduction/Takes/Source/TakeRecorder/Private/Customization/RecorderPropertyMapCustomization.h
new file mode 100644
--- /dev/null
+++ b/Engine/Plugins/VirtualProduction/Takes/Source/TakeRecorder/Private/Customization/RecorderPropertyMapCustomization.h
@@ -0,0 +1,37 @@
+﻿// Copyright Epic Games, Inc. All Rights Reserved.
+
+#pragma once
+
+#include "IPropertyTypeCustomization.h"
+#include "Templates/SharedPointerFwd.h"
+
+class FString;
+class IDetailGroup;
+class IPropertyHandle;
+class IPropertyHandleArray;
+enum class ECheckBoxState : uint8;
+
+namespace UE::TakeRecorder
+{
+class FRecorderPropertyMapCustomization : public IPropertyTypeCustomization
+{
+public:
+	
+	static const FString PropertyPathDelimiter; 
+
+	//~ Begin IPropertyTypeCustomization Interface
+	virtual void CustomizeHeader(TSharedRef<IPropertyHandle> PropertyHandle, FDetailWidgetRow& HeaderRow, IPropertyTypeCustomizationUtils& CustomizationUtils);
+	virtual IDetailGroup& GetOrCreateDetailGroup(IDetailChildrenBuilder& ChildBuilder, TMap<FString, IDetailGroup*>& GroupMap, TSharedPtr<IPropertyHandleArray> PropertiesArray, FString& GroupName);
+	virtual void CustomizeChildren(TSharedRef<IPropertyHandle> PropertyHandle, IDetailChildrenBuilder& ChildBuilder, IPropertyTypeCustomizationUtils& CustomizationUtils);
+	//~ End IPropertyTypeCustomization Interface
+
+private:
+	
+	void OnGroupCheckStateChanged( ECheckBoxState InNewState, TSharedPtr<IPropertyHandleArray> RecordedPropertiesArrayHandle, FString GroupName) const;
+	ECheckBoxState OnGroupGetCheckState( TSharedPtr<IPropertyHandleArray> RecordedPropertiesArrayHandle, FString GroupName) const;
+	
+	void OnCheckStateChanged( ECheckBoxState InNewState, TSharedRef<IPropertyHandle> PropertyHandle );
+	ECheckBoxState OnGetCheckState( TSharedRef<IPropertyHandle> PropertyHandle ) const;
+};
+}
+
diff --git a/Engine/Plugins/VirtualProduction/Takes/Source/TakeRecorder/Private/Customization/RecorderSourceObjectCustomization.cpp b/Engine/Plugins/VirtualProduction/Takes/Source/TakeRecorder/Private/Customization/RecorderSourceObjectCustomization.cpp
new file mode 100644
--- /dev/null
+++ b/Engine/Plugins/VirtualProduction/Takes/Source/TakeRecorder/Private/Customization/RecorderSourceObjectCustomization.cpp
@@ -0,0 +1,65 @@
+﻿// Copyright Epic Games, Inc. All Rights Reserved.
+
+#include "RecorderSourceObjectCustomization.h"
+
+#include "DetailCategoryBuilder.h"
+#include "DetailLayoutBuilder.h"
+#include "TakeRecorderSource.h"
+
+#define LOCTEXT_NAMESPACE "SLevelSequenceTakeEditor"
+
+namespace UE::TakeRecorder
+{
+void FRecorderSourceObjectCustomization::CustomizeDetails(IDetailLayoutBuilder& DetailBuilder)
+{
+	FText NewTitle = ComputeTitle(DetailBuilder.GetDetailsViewSharedPtr());
+	if (!NewTitle.IsEmpty())
+	{
+		// Edit the category and add *all* properties for the object to it
+		IDetailCategoryBuilder& CategoryBuilder = DetailBuilder.EditCategory("CustomCategory", NewTitle);
+
+		UClass* BaseClass = DetailBuilder.GetBaseClass();
+		while (BaseClass)
+		{
+			for (FProperty* Property : TFieldRange<FProperty>(BaseClass, EFieldIteratorFlags::ExcludeSuper))
+			{
+				CategoryBuilder.AddProperty(Property->GetFName(), BaseClass);
+			}
+
+			BaseClass = BaseClass->GetSuperClass();
+		}
+	}
+}
+
+FText FRecorderSourceObjectCustomization::ComputeTitle(const TSharedPtr<const IDetailsView>& DetailsView) const
+{
+	if (!DetailsView)
+	{
+		return FText();
+	}
+
+	// Compute the title for all the sources that this details panel is editing
+	static const FName CategoryName = "Category";
+
+	const TArray<TWeakObjectPtr<UObject>>& SelectedObjects = DetailsView->GetSelectedObjects();
+	UObject* FirstObject = SelectedObjects.Num() > 0 ? SelectedObjects[0].Get() : nullptr;
+
+	if (FirstObject)
+	{
+		if (SelectedObjects.Num() == 1)
+		{
+			UTakeRecorderSource* Source = Cast<UTakeRecorderSource>(FirstObject);
+			return Source ? Source->GetDisplayText() : FText::FromString(FirstObject->GetName());
+		}
+		else if (SelectedObjects.Num() > 1)
+		{
+			const FString& Category = FirstObject->GetClass()->GetMetaData(CategoryName);
+			return FText::Format(LOCTEXT("CategoryFormatString", "{0} ({1})"), FText::FromString(Category), SelectedObjects.Num());
+		}
+	}
+
+	return FText();
+}
+}
+
+#undef LOCTEXT_NAMESPACE
\ No newline at end of file
diff --git a/Engine/Plugins/VirtualProduction/Takes/Source/TakeRecorder/Private/Customization/RecorderSourceObjectCustomization.h b/Engine/Plugins/VirtualProduction/Takes/Source/TakeRecorder/Private/Customization/RecorderSourceObjectCustomization.h
new file mode 100644
--- /dev/null
+++ b/Engine/Plugins/VirtualProduction/Takes/Source/TakeRecorder/Private/Customization/RecorderSourceObjectCustomization.h
@@ -0,0 +1,22 @@
+﻿// Copyright Epic Games, Inc. All Rights Reserved.
+
+#pragma once
+
+#include "IDetailCustomization.h"
+#include "IDetailsView.h"
+
+namespace UE::TakeRecorder
+{
+class FRecorderSourceObjectCustomization : public IDetailCustomization
+{
+public:
+
+	//~ Begin IDetailCustomization Interface
+	virtual void CustomizeDetails(IDetailLayoutBuilder& DetailBuilder) override;
+	//~ End IDetailCustomization Interface
+
+private:
+	
+	FText ComputeTitle(const TSharedPtr<const IDetailsView>& DetailsView) const;
+};
+}
diff --git a/Engine/Plugins/VirtualProduction/Takes/Source/TakeRecorder/Private/Customization/TakeRecorderAudioSettingsCustomization.cpp b/Engine/Plugins/VirtualProduction/Takes/Source/TakeRecorder/Private/Customization/TakeRecorderAudioSettingsCustomization.cpp
new file mode 100644
--- /dev/null
+++ b/Engine/Plugins/VirtualProduction/Takes/Source/TakeRecorder/Private/Customization/TakeRecorderAudioSettingsCustomization.cpp
@@ -0,0 +1,360 @@
+// Copyright Epic Games, Inc. All Rights Reserved.
+
+#include "TakeRecorderAudioSettingsCustomization.h"
+#include "IDetailChildrenBuilder.h"
+#include "TakeRecorderSettings.h"
+#include "TakeRecorderSourceProperty.h"
+
+#include "Framework/MultiBox/MultiBoxBuilder.h"
+#include "Widgets/Input/SCheckBox.h"
+
+#define LOCTEXT_NAMESPACE "FAudioInputDevicePropertyCustomization"
+
+void FAudioInputDevicePropertyCustomization::CustomizeChildren(TSharedRef<IPropertyHandle> PropertyHandle, IDetailChildrenBuilder& ChildBuilder, IPropertyTypeCustomizationUtils& CustomizationUtils)
+{
+	UseSystemDefaultHandle = PropertyHandle->GetChildHandle(GET_MEMBER_NAME_CHECKED(FAudioInputDeviceProperty, bUseSystemDefaultAudioDevice));
+	InputDeviceHandle = PropertyHandle->GetChildHandle(GET_MEMBER_NAME_CHECKED(FAudioInputDeviceProperty, DeviceId));
+	BufferSizeHandle = PropertyHandle->GetChildHandle(GET_MEMBER_NAME_CHECKED(FAudioInputDeviceProperty, AudioInputBufferSize));
+	DeviceInfoArrayHandle = PropertyHandle->GetChildHandle(GET_MEMBER_NAME_CHECKED(FAudioInputDeviceProperty, DeviceInfoArray))->AsArray();
+
+	// Create check box which indicates system default audio device should be used (disables device combo button)
+	if (UseSystemDefaultHandle.IsValid())
+	{
+		IDetailPropertyRow& UseSystemDefaultRow = ChildBuilder.AddProperty(PropertyHandle);
+		UseSystemDefaultRow.CustomWidget()
+		.NameContent()
+		[
+			UseSystemDefaultHandle->CreatePropertyNameWidget()
+		]
+		.ValueContent()
+		[
+			SNew(SCheckBox)
+			.IsChecked_Lambda([&]() -> ECheckBoxState
+			{
+				bool bUseSystemDefault;
+				UseSystemDefaultHandle->GetValue(bUseSystemDefault);
+
+				return bUseSystemDefault ? ECheckBoxState::Checked : ECheckBoxState::Unchecked;
+			})
+			.OnCheckStateChanged_Lambda([&](ECheckBoxState NewState) -> void
+			{
+				bool bUseSystemDefault = (NewState == ECheckBoxState::Checked);
+				UseSystemDefaultHandle->SetValue(bUseSystemDefault);
+
+				// Refresh the device list
+				RequestsDeviceListRefresh();
+			})
+		];
+	}
+
+	// Create audio input device combo button
+	if (InputDeviceHandle.IsValid() && DeviceInfoArrayHandle.IsValid())
+	{
+		IDetailPropertyRow& InputDevicePropertyRow = ChildBuilder.AddProperty(PropertyHandle);
+
+		DeviceInfoArrayHandle->SetOnNumElementsChanged(FSimpleDelegate::CreateSP(this, &FAudioInputDevicePropertyCustomization::OnDeviceListChanged));
+		BuildAudioInputDeviceList();
+
+		InputDevicePropertyRow.CustomWidget()
+		.NameContent()
+		[
+			InputDeviceHandle->CreatePropertyNameWidget()
+		]
+		.ValueContent()
+		[
+			CreateDeviceListComboButton()
+		];
+	}
+
+	// Create audio callback buffer size widget
+	if (BufferSizeHandle.IsValid())
+	{
+		IDetailPropertyRow& BufferSizePropertyRow = ChildBuilder.AddProperty(PropertyHandle);
+		BufferSizePropertyRow.CustomWidget()
+		.NameContent()
+		[
+			BufferSizeHandle->CreatePropertyNameWidget()
+		]
+		.ValueContent()
+		[
+			BufferSizeHandle->CreatePropertyValueWidget()
+		];
+	}
+
+	SynchronizeWidgetStates();
+
+	// Register delegate so UI can update when the audio input device changes
+	// This is needed because the device menu lives in both the project settings window and the take recorder panel
+	if (UTakeRecorderAudioInputSettings* AudioInputSettings = TakeRecorderAudioSettingsUtils::GetTakeRecorderAudioInputSettings())
+	{
+		AudioInputSettings->GetOnAudioInputDeviceChanged().Add(FSimpleDelegate::CreateSP(this, &FAudioInputDevicePropertyCustomization::SynchronizeWidgetStates));
+	}
+}
+
+FString FAudioInputDevicePropertyCustomization::GetDeviceNameFromInfoProperty(const TSharedPtr<IPropertyHandle>& InDeviceInfoHandle)
+{
+	FString DeviceNameStr;
+	if (InDeviceInfoHandle && InDeviceInfoHandle.IsValid())
+	{
+		TSharedPtr<IPropertyHandle> DeviceNameHandle = InDeviceInfoHandle->GetChildHandle(GET_MEMBER_NAME_CHECKED(FAudioInputDeviceInfoProperty, DeviceName));
+		if (DeviceNameHandle.IsValid())
+		{
+			DeviceNameHandle->GetValue(DeviceNameStr);
+		}
+	}
+
+	return DeviceNameStr;
+}
+
+FString FAudioInputDevicePropertyCustomization::GetDeviceIdFromInfoProperty(const TSharedPtr<IPropertyHandle>& InDeviceInfoHandle)
+{
+	FString DeviceIdStr;
+	if (InDeviceInfoHandle.IsValid())
+	{
+		TSharedPtr<IPropertyHandle> DeviceIdHandle = InDeviceInfoHandle->GetChildHandle(GET_MEMBER_NAME_CHECKED(FAudioInputDeviceInfoProperty, DeviceId));
+		if (DeviceIdHandle.IsValid())
+		{
+			DeviceIdHandle->GetValue(DeviceIdStr);
+		}
+	}
+
+	return DeviceIdStr;
+}
+
+int32 FAudioInputDevicePropertyCustomization::GetChannelCountFromInfoProperty(const TSharedPtr<IPropertyHandle>& InDeviceInfoHandle)
+{
+	int32 ChannelCount = 0;
+	if (InDeviceInfoHandle.IsValid())
+	{
+		TSharedPtr<IPropertyHandle> ChannelCountHandle = InDeviceInfoHandle->GetChildHandle(GET_MEMBER_NAME_CHECKED(FAudioInputDeviceInfoProperty, InputChannels));
+		if (ChannelCountHandle.IsValid())
+		{
+			ChannelCountHandle->GetValue(ChannelCount);
+		}
+	}
+
+	return ChannelCount;
+}
+
+bool FAudioInputDevicePropertyCustomization::GetIsDefaultFromInfoProperty(const TSharedPtr<IPropertyHandle>& InDeviceInfoHandle)
+{
+	bool bIsDefaultDevice = false;
+	if (InDeviceInfoHandle.IsValid())
+	{
+		TSharedPtr<IPropertyHandle> IsDefaultHandle = InDeviceInfoHandle->GetChildHandle(GET_MEMBER_NAME_CHECKED(FAudioInputDeviceInfoProperty, bIsDefaultDevice));
+		if (IsDefaultHandle.IsValid())
+		{
+			IsDefaultHandle->GetValue(bIsDefaultDevice);
+		}
+	}
+
+	return bIsDefaultDevice;
+}
+
+void FAudioInputDevicePropertyCustomization::BuildAudioInputDeviceList()
+{
+	FString CurrentDeviceId;
+	InputDeviceHandle->GetValue(CurrentDeviceId);
+
+	bool bHaveCurrentDevice = CurrentDeviceId.Len() != 0;
+	uint32 NumDevices;
+	DeviceInfoArrayHandle->GetNumElements(NumDevices);
+
+	AudioInputDevices.Empty();
+
+	for (uint32 DeviceIndex = 0; DeviceIndex < NumDevices; ++DeviceIndex)
+	{
+		TSharedRef<IPropertyHandle> DeviceInfoHandle = DeviceInfoArrayHandle->GetElement(DeviceIndex);
+		if (DeviceInfoHandle->IsValidHandle())
+		{
+			bool bIsDefaultDevice = GetIsDefaultFromInfoProperty(DeviceInfoHandle);
+			FString DeviceId = GetDeviceIdFromInfoProperty(DeviceInfoHandle);
+
+			if (bIsDefaultDevice)
+			{
+				DefaultDeviceId = DeviceId;
+			}
+
+			if (bIsDefaultDevice && !bHaveCurrentDevice)
+			{
+				CurrentDeviceId = DeviceId;
+				InputDeviceHandle->SetValue(CurrentDeviceId);
+				SelectedAudioInputDevice = DeviceInfoHandle;
+			}
+			else if (bHaveCurrentDevice)
+			{
+				if (CurrentDeviceId == DeviceId)
+				{
+					SelectedAudioInputDevice = DeviceInfoHandle;
+				}
+			}
+
+			AudioInputDevices.Add(DeviceInfoHandle);
+		}
+	}
+}
+
+void FAudioInputDevicePropertyCustomization::RequestsDeviceListRefresh()
+{
+	TakeRecorderAudioSettingsUtils::RefreshAudioInputSettings();
+
+	BuildAudioInputDeviceList();
+	SynchronizeWidgetStates();
+}
+
+void FAudioInputDevicePropertyCustomization::OnDeviceListChanged()
+{
+	BuildAudioInputDeviceList();
+	SynchronizeWidgetStates();
+}
+
+void FAudioInputDevicePropertyCustomization::MenuItemRefreshSelected()
+{
+	RequestsDeviceListRefresh();
+}
+
+void FAudioInputDevicePropertyCustomization::MenuItemDeviceSelected(TSharedPtr<IPropertyHandle> InDeviceInfoHandle)
+{
+	SelectedAudioInputDevice = InDeviceInfoHandle;
+	InputDeviceHandle->SetValue(GetDeviceIdFromInfoProperty(InDeviceInfoHandle));
+}
+
+TSharedRef<SWidget> FAudioInputDevicePropertyCustomization::OnGenerateDeviceMenu()
+{
+	FMenuBuilder MenuBuilder(true, nullptr, nullptr);
+
+	MenuBuilder.AddMenuEntry(
+		LOCTEXT("RefreshDeviceList", "Refresh List"),
+		LOCTEXT("RefreshDeviceListToolTip", "Refresh contents of this menu's Audio Input Device List"),
+		FSlateIcon(),
+		FUIAction(
+			FExecuteAction::CreateSP(this, &FAudioInputDevicePropertyCustomization::MenuItemRefreshSelected)
+		)
+	);
+
+	MenuBuilder.BeginSection(NAME_None, LOCTEXT("AudioInputDeviceMenuSection", "Audio Input Devices"));
+	{
+		for (TSharedPtr<IPropertyHandle> DeviceInfoHandle : AudioInputDevices)
+		{
+			FString DeviceName = GetDeviceNameFromInfoProperty(DeviceInfoHandle);
+
+			if (!DeviceName.IsEmpty())
+			{
+				MenuBuilder.AddMenuEntry(
+					FText::FromString(DeviceName),
+					FText(),
+					FSlateIcon(),
+					FUIAction(
+						FExecuteAction::CreateSP(this, &FAudioInputDevicePropertyCustomization::MenuItemDeviceSelected, DeviceInfoHandle)
+					)
+				);
+			}
+		}
+	}
+	MenuBuilder.EndSection();
+
+	return MenuBuilder.MakeWidget();
+}
+
+TSharedRef<SWidget> FAudioInputDevicePropertyCustomization::CreateDeviceListComboButton()
+{
+	TSharedRef<SSimpleComboButton> NewWidget = SNew(SSimpleComboButton)
+		.OnGetMenuContent(this, &FAudioInputDevicePropertyCustomization::OnGenerateDeviceMenu)
+		.ToolTipText(LOCTEXT("AudioInputDeviceMenuToolTip", "Available Audio Input Devices"))
+		.HasDownArrow(true)
+		.UsesSmallText(true)
+		.Text_Lambda([this]()
+		{
+			return FText::AsCultureInvariant(GetDeviceNameFromInfoProperty(SelectedAudioInputDevice));
+		});
+
+	DeviceListComboButton = NewWidget;
+
+	return NewWidget;
+}
+
+void FAudioInputDevicePropertyCustomization::UpdateDeviceCombobxEnabled()
+{
+	// Get current data model value
+	bool bUseSystemDefault;
+	UseSystemDefaultHandle->GetValue(bUseSystemDefault);
+
+	// We want the widget enabled if bUseSystemDefault is false
+	bool bEnableWidget = !bUseSystemDefault;
+
+	// Get the current widget state
+	bool bIsWidgetEnabled = DeviceListComboButton->IsEnabled();
+
+	// If the UI doesn't match the data model, refresh the UI
+	if (bEnableWidget != bIsWidgetEnabled)
+	{
+		DeviceListComboButton->SetEnabled(bEnableWidget);
+        
+        // If disabling, set the current menu item to default
+        if (!bEnableWidget)
+        {
+            InputDeviceHandle->SetValue(DefaultDeviceId);
+        }
+	}
+}
+
+void FAudioInputDevicePropertyCustomization::SynchronizeWidgetStates()
+{
+	UpdateDeviceCombobxEnabled();
+
+	uint32 NumDevices;
+	DeviceInfoArrayHandle->GetNumElements(NumDevices);
+
+	// Get current data model value
+	FString CurrentDeviceId;
+	InputDeviceHandle->GetValue(CurrentDeviceId);
+
+	// If the UI doesn't match the data model, refresh the UI
+	if (CurrentDeviceId != GetDeviceIdFromInfoProperty(SelectedAudioInputDevice))
+	{
+		for (uint32 DeviceIndex = 0; DeviceIndex < NumDevices; ++DeviceIndex)
+		{
+			TSharedRef<IPropertyHandle> DeviceInfoHandle = DeviceInfoArrayHandle->GetElement(DeviceIndex);
+			if (DeviceInfoHandle->IsValidHandle())
+			{
+				FString DeviceId = GetDeviceIdFromInfoProperty(DeviceInfoHandle);
+
+				if (DeviceId == CurrentDeviceId)
+				{
+					SelectedAudioInputDevice = DeviceInfoHandle;
+					DeviceListComboButton->Invalidate(EInvalidateWidgetReason::Paint);
+					break;
+				}
+			}
+		}
+	}
+}
+
+UTakeRecorderAudioInputSettings* TakeRecorderAudioSettingsUtils::GetTakeRecorderAudioInputSettings()
+{
+	const UTakeRecorderProjectSettings* ProjectSettings = GetDefault<UTakeRecorderProjectSettings>();
+	for (TWeakObjectPtr<UObject> WeakObject : ProjectSettings->AdditionalSettings)
+	{
+		UObject* Object = WeakObject.Get();
+		if (Object)
+		{
+			if (Object->IsA<UTakeRecorderAudioInputSettings>())
+			{
+				return Cast<UTakeRecorderAudioInputSettings>(Object);
+			}
+		}
+	}
+	return nullptr;
+}
+
+void TakeRecorderAudioSettingsUtils::RefreshAudioInputSettings()
+{
+	// Enumerate audio devices which will rebuild the device list. Note, this can be expensive depending on the hardware
+	// attached to the machine.
+	if (UTakeRecorderAudioInputSettings* AudioInputSettings = TakeRecorderAudioSettingsUtils::GetTakeRecorderAudioInputSettings())
+	{
+		AudioInputSettings->EnumerateAudioDevices(true);
+	}
+}
+
+#undef LOCTEXT_NAMESPACE // "FAudioInputDevicePropertyCustomization"
diff --git a/Engine/Plugins/VirtualProduction/Takes/Source/TakeRecorder/Private/Customization/TakeRecorderAudioSettingsCustomization.h b/Engine/Plugins/VirtualProduction/Takes/Source/TakeRecorder/Private/Customization/TakeRecorderAudioSettingsCustomization.h
new file mode 100644
--- /dev/null
+++ b/Engine/Plugins/VirtualProduction/Takes/Source/TakeRecorder/Private/Customization/TakeRecorderAudioSettingsCustomization.h
@@ -0,0 +1,57 @@
+// Copyright Epic Games, Inc. All Rights Reserved.
+
+#pragma once
+
+#include "DetailWidgetRow.h"
+#include "IPropertyTypeCustomization.h"
+#include "TakeRecorderSourceProperty.h"
+
+#include "SSimpleComboButton.h"
+#include "Widgets/Text/STextBlock.h"
+
+class FAudioInputDevicePropertyCustomization : public IPropertyTypeCustomization
+{
+public:
+
+	virtual void CustomizeHeader(TSharedRef<IPropertyHandle> PropertyHandle, FDetailWidgetRow& HeaderRow, IPropertyTypeCustomizationUtils& CustomizationUtils) override {}
+
+	virtual void CustomizeChildren(TSharedRef<IPropertyHandle> PropertyHandle, IDetailChildrenBuilder& ChildBuilder, IPropertyTypeCustomizationUtils& CustomizationUtils) override;
+
+private:
+
+	static FString GetDeviceNameFromInfoProperty(const TSharedPtr<IPropertyHandle>& InDeviceInfoHandle);
+
+	FString GetDeviceIdFromInfoProperty(const TSharedPtr<IPropertyHandle>& InDeviceInfoHandle);
+	int32 GetChannelCountFromInfoProperty(const TSharedPtr<IPropertyHandle>& InDeviceInfoHandle);
+
+	bool GetIsDefaultFromInfoProperty(const TSharedPtr<IPropertyHandle>& InDeviceInfoHandle);
+	void BuildAudioInputDeviceList();
+	TSharedRef<SWidget> CreateDeviceListComboButton();
+
+	void RequestsDeviceListRefresh();
+	TSharedRef<SWidget> OnGenerateDeviceMenu();
+	void MenuItemRefreshSelected();
+	void MenuItemDeviceSelected(TSharedPtr<IPropertyHandle> InDeviceInfoHandle);
+
+	void OnDeviceListChanged();
+	void UpdateDeviceCombobxEnabled();
+	void SynchronizeWidgetStates();
+
+	TSharedPtr<IPropertyHandle> UseSystemDefaultHandle;
+	TSharedPtr<IPropertyHandle> InputDeviceHandle;
+	TSharedPtr<IPropertyHandle> BufferSizeHandle;
+	TSharedPtr<IPropertyHandleArray> DeviceInfoArrayHandle;
+
+	TSharedPtr<IPropertyHandle> SelectedAudioInputDevice;
+	TArray<TSharedPtr<IPropertyHandle>> AudioInputDevices;
+	TSharedPtr<SSimpleComboButton> DeviceListComboButton;
+	FString DefaultDeviceId;
+};
+
+class TakeRecorderAudioSettingsUtils
+{
+public:
+
+	static UTakeRecorderAudioInputSettings* GetTakeRecorderAudioInputSettings();
+	static void RefreshAudioInputSettings();
+};
\ No newline at end of file
diff --git a/Engine/Plugins/VirtualProduction/Takes/Source/TakeRecorder/Private/Customization/TakeRecorderProjectSettingsCustomization.h b/Engine/Plugins/VirtualProduction/Takes/Source/TakeRecorder/Private/Customization/TakeRecorderProjectSettingsCustomization.h
new file mode 100644
--- /dev/null
+++ b/Engine/Plugins/VirtualProduction/Takes/Source/TakeRecorder/Private/Customization/TakeRecorderProjectSettingsCustomization.h
@@ -0,0 +1,70 @@
+// Copyright Epic Games, Inc. All Rights Reserved.
+
+#pragma once
+
+#include "Customization/TakeRecorderAudioSettingsCustomization.h"
+#include "IDetailChildrenBuilder.h"
+#include "IDetailCustomization.h"
+#include "DetailLayoutBuilder.h"
+#include "DetailCategoryBuilder.h"
+#include "DetailWidgetRow.h"
+#include "Modules/ModuleManager.h"
+#include "TakeRecorderSettings.h"
+#include "TakeRecorderSourceProperty.h"
+
+class FTakeRecorderProjectSettingsCustomization : public IDetailCustomization
+{
+	virtual void CustomizeDetails(IDetailLayoutBuilder& DetailLayout) override
+	{
+		FPropertyEditorModule& PropertyEditorModule = FModuleManager::Get().LoadModuleChecked<FPropertyEditorModule>("PropertyEditor");
+
+		// Register audio device details customization for the project settings window
+		PropertyEditorModule.RegisterCustomPropertyTypeLayout(FName(TEXT("AudioInputDeviceProperty")), FOnGetPropertyTypeCustomizationInstance::CreateLambda(&MakeShared<FAudioInputDevicePropertyCustomization>));
+
+		// Pop the take recorder category to the top
+		DetailLayout.EditCategory("Take Recorder");
+
+		TArray<TWeakObjectPtr<UObject>> CustomizedObjects;
+		DetailLayout.GetObjectsBeingCustomized(CustomizedObjects);
+
+		for (TWeakObjectPtr<UObject> EditObject : CustomizedObjects)
+		{
+			UTakeRecorderProjectSettings* Settings = Cast<UTakeRecorderProjectSettings>(EditObject.Get());
+			if (!Settings)
+			{
+				continue;
+			}
+
+			for (TWeakObjectPtr<UObject> WeakAdditionalSettings : Settings->AdditionalSettings)
+			{
+				UObject* AdditionalSettings = WeakAdditionalSettings.Get();
+				if (!AdditionalSettings)
+				{
+					continue;
+				}
+
+				UClass* Class = AdditionalSettings->GetClass();
+
+				TArray<FProperty*> EditProperties;
+				for (FProperty* Property : TFieldRange<FProperty>(Class))
+				{
+					if (Property && Property->HasAllPropertyFlags(CPF_Edit | CPF_Config))
+					{
+						EditProperties.Add(Property);
+					}
+				}
+
+				if (EditProperties.Num() > 0)
+				{
+					IDetailCategoryBuilder& Category = DetailLayout.EditCategory(*Class->GetDisplayNameText().ToString());
+
+					TArray<UObject*> SettingAsArray = { AdditionalSettings };
+					for (FProperty* Property : EditProperties)
+					{
+						Category.AddExternalObjectProperty(SettingAsArray, Property->GetFName());
+					}
+				}
+			}
+		}
+	}
+};
diff --git a/Engine/Plugins/VirtualProduction/Takes/Source/TakeRecorder/Private/TakeRecorderModule.cpp b/Engine/Plugins/VirtualProduction/Takes/Source/TakeRecorder/Private/TakeRecorderModule.cpp
--- a/Engine/Plugins/VirtualProduction/Takes/Source/TakeRecorder/Private/TakeRecorderModule.cpp
+++ b/Engine/Plugins/VirtualProduction/Takes/Source/TakeRecorder/Private/TakeRecorderModule.cpp
@@ -3,7 +3,7 @@
 #include "TakeRecorderModule.h"
 #include "TakeRecorderSettings.h"
 #include "SequencerUtilities.h"
-#include "TakeRecorderProjectSettingsCustomization.h"
+#include "Customization/TakeRecorderProjectSettingsCustomization.h"
 #include "Modules/ModuleManager.h"
 #include "LevelEditor.h"
 #include "TakeRecorderCommands.h"
diff --git a/Engine/Plugins/VirtualProduction/Takes/Source/TakeRecorder/Private/TakeRecorderProjectSettingsCustomization.h b/Engine/Plugins/VirtualProduction/Takes/Source/TakeRecorder/Private/TakeRecorderProjectSettingsCustomization.h
deleted file mode 100644
diff --git a/Engine/Plugins/VirtualProduction/Takes/Source/TakeRecorder/Private/Widgets/SLevelSequenceTakeEditor.cpp b/Engine/Plugins/VirtualProduction/Takes/Source/TakeRecorder/Private/Widgets/SLevelSequenceTakeEditor.cpp
--- a/Engine/Plugins/VirtualProduction/Takes/Source/TakeRecorder/Private/Widgets/SLevelSequenceTakeEditor.cpp
+++ b/Engine/Plugins/VirtualProduction/Takes/Source/TakeRecorder/Private/Widgets/SLevelSequenceTakeEditor.cpp
@@ -1,8 +1,11 @@
 // Copyright Epic Games, Inc. All Rights Reserved.
 
 #include "Widgets/SLevelSequenceTakeEditor.h"
-#include "Widgets/TakeRecorderWidgetConstants.h"
-#include "Widgets/TakeRecorderAudioSettingsCustomization.h"
+#include "Customization/AudioInputChannelPropertyCustomization.h"
+#include "Customization/TakeRecorderAudioSettingsCustomization.h"
+#include "Customization/RecordedPropertyCustomization.h"
+#include "Customization/RecorderPropertyMapCustomization.h"
+#include "Customization/RecorderSourceObjectCustomization.h"
 #include "TakePreset.h"
 #include "TakeMetaData.h"
 #include "TakeRecorderSource.h"
@@ -11,14 +14,11 @@
 #include "TakeRecorderStyle.h"
 #include "Widgets/STakeRecorderSources.h"
 #include "TakeRecorderModule.h"
-#include "TakeRecorderSettings.h"
-
 #include "LevelSequence.h"
 
 // Core includes
 #include "Modules/ModuleManager.h"
 #include "Algo/Sort.h"
-#include "Misc/FileHelper.h"
 #include "UObject/UObjectIterator.h"
 #include "Templates/SubclassOf.h"
 #include "ClassIconFinder.h"
@@ -35,10 +35,6 @@
 #include "AssetToolsModule.h"
 
 // Slate includes
-#include "Widgets/SBoxPanel.h"
-#include "Widgets/Text/STextBlock.h"
-#include "Widgets/Input/SButton.h"
-#include "Widgets/Input/SComboButton.h"
 #include "Widgets/Input/SComboBox.h"
 #include "Widgets/Input/SEditableTextBox.h"
 #include "Widgets/Layout/SSplitter.h"
@@ -51,7 +47,6 @@
 
 // Style includes
 #include "Styling/AppStyle.h"
-#include "EditorFontGlyphs.h"
 
 // UnrealEd includes
 #include "ScopedTransaction.h"
@@ -59,19 +54,8 @@
 // PropertyEditor includes
 #include "PropertyEditorModule.h"
 #include "IDetailsView.h"
-#include "IDetailGroup.h"
-#include "DetailLayoutBuilder.h"
-#include "DetailCategoryBuilder.h"
-#include "IDetailChildrenBuilder.h"
-#include "IDetailCustomization.h"
-#include "IDetailRootObjectCustomization.h"
-#include "IPropertyUtilities.h"
 #include "DetailWidgetRow.h"
 
-// Widget Includes
-#include "Widgets/Input/SCheckBox.h"
-#include "Widgets/Images/SImage.h"
-
 #define LOCTEXT_NAMESPACE "SLevelSequenceTakeEditor"
 
 
@@ -298,600 +282,6 @@ void SLevelSequenceTakeEditor::OnDetailsPropertiesChanged(const FPropertyChanged
 	OnDetailsPropertiesChangedEvent.ExecuteIfBound(InEvent);
 }
 
-class FRecordedPropertyCustomization : public IPropertyTypeCustomization
-{
-	static const FString PropertyPathDelimiter;
-
-   	virtual void CustomizeHeader( TSharedRef<IPropertyHandle> PropertyHandle, FDetailWidgetRow& HeaderRow, IPropertyTypeCustomizationUtils& CustomizationUtils )
-   	{
-	   	if (PropertyHandle->IsValidHandle())
-	   	{
-	        TSharedPtr<IPropertyHandle> PropertyNameHandle = PropertyHandle->GetChildHandle(GET_MEMBER_NAME_CHECKED(FActorRecordedProperty, PropertyName)); 
-	        TSharedPtr<IPropertyHandle> bEnabledHandle = PropertyHandle->GetChildHandle(GET_MEMBER_NAME_CHECKED(FActorRecordedProperty, bEnabled)); 
-	        TSharedPtr<IPropertyHandle> RecorderNameHandle = PropertyHandle->GetChildHandle(GET_MEMBER_NAME_CHECKED(FActorRecordedProperty, RecorderName));
-
-	        FString PropertyNameValue;
-	        PropertyNameHandle->GetValueAsDisplayString(PropertyNameValue);
-
-		    FString ParentGroups; 
-		    FString PropertyName;
-		    FText DisplayString = FText::FromString(PropertyNameValue);
-		    if (PropertyNameValue.Split(PropertyPathDelimiter, &ParentGroups, &PropertyName, ESearchCase::IgnoreCase, ESearchDir::FromEnd))
-		    {
-		    	DisplayString = FText::FromString(PropertyName);
-			}
-
-	        HeaderRow 
-	        [
-
-	        	SNew(SHorizontalBox)
-
-	        	+SHorizontalBox::Slot()
-	        	.AutoWidth()
-	        	[
-	        		bEnabledHandle->CreatePropertyValueWidget(false)
-	        	]
-
-	        	+SHorizontalBox::Slot()
-	        	.Padding(8, 0, 0, 0)
-	        	[
-	        		PropertyNameHandle->CreatePropertyNameWidget(DisplayString)
-	        	]
-
-	        ];
-	   	}
-   		return;
-   	}
-
-	virtual void CustomizeChildren( TSharedRef<IPropertyHandle> PropertyHandle, IDetailChildrenBuilder& ChildBuilder, IPropertyTypeCustomizationUtils& CustomizationUtils )
-	{
-		// Intentionally Left Blank, Child Customization was handled in the Header Row
-	}
-};
-
-const FString FRecordedPropertyCustomization::PropertyPathDelimiter = FString(TEXT("."));
-
-class FRecorderPropertyMapCustomization : public IPropertyTypeCustomization
-{
-	static const FString PropertyPathDelimiter; 
-
-	virtual void CustomizeHeader( TSharedRef<IPropertyHandle> PropertyHandle, FDetailWidgetRow& HeaderRow, IPropertyTypeCustomizationUtils& CustomizationUtils )
-	{
-	    TSharedPtr<IPropertyHandle> RecordedObjectHandle = PropertyHandle->GetChildHandle(GET_MEMBER_NAME_CHECKED(UActorRecorderPropertyMap, RecordedObject)); 
-
-	    UObject* Object = nullptr;
-	    FText ActorOrComponentName = LOCTEXT("MissingActorOrComponentName", "MissingActorOrComponentName");
-
-	    const FSlateBrush* Icon = nullptr;
-
-	    bool bSuccess = false;
-		if ( RecordedObjectHandle->IsValidHandle() && RecordedObjectHandle->GetValue(Object) == FPropertyAccess::Success && Object != nullptr )
-		{
-			bSuccess = true;
-
-			if (AActor* Actor = Cast<AActor>(Object))
-			{
-				ActorOrComponentName = FText::AsCultureInvariant(Actor->GetActorLabel());
-				Icon = FClassIconFinder::FindIconForActor(Actor);
-			}
-			else
-			{
-				ActorOrComponentName = FText::AsCultureInvariant(Object->GetName());
-				Icon = FSlateIconFinder::FindIconBrushForClass(Object->GetClass());
-			}
-		}
-
-	    HeaderRow
-	    [
-	    	SNew(SHorizontalBox)
-
-	    	+SHorizontalBox::Slot()
-	    	.AutoWidth()
-	    	.VAlign(VAlign_Center)
-	    	[
-	    		SNew( SCheckBox )
-	    		.OnCheckStateChanged( this, &FRecorderPropertyMapCustomization::OnCheckStateChanged, PropertyHandle )	
-	    		.IsChecked( this, &FRecorderPropertyMapCustomization::OnGetCheckState, PropertyHandle )	
-	    		.Padding(0.0)
-	    	]
-
-	    	+SHorizontalBox::Slot()
-	    	.AutoWidth()
-	    	.VAlign(VAlign_Center)
-	    	.Padding(8, 0, 0, 0)
-	    	[
-	    		SNew(SImage)
-				.Image(Icon)
-	    	]
-
-	    	+SHorizontalBox::Slot()
-	    	.VAlign(VAlign_Center)
-	    	.AutoWidth()
-	    	.Padding(2.0, 0.0)
-	    	[
-	    		SNew(STextBlock)
-	    		.Text(ActorOrComponentName)
-	    		.Font( FAppStyle::GetFontStyle( "PropertyWindow.BoldFont" ))
-	    	]
-
-	    	+SHorizontalBox::Slot()
-	    	.FillWidth(1.0f)
-	    	.VAlign(VAlign_Center)
-	    	.Padding(2.0, 0.0)
-	    	[
-	    		SNew(STextBlock)
-	    		.Text(LOCTEXT("TakeRecorderRecordedPropertiesTitle", "Recorded Properties"))
-	    		.Font( FAppStyle::GetFontStyle( "PropertyWindow.NormalFont" ))
-	    	]
-	    ];
-	}
-
-	virtual IDetailGroup& GetOrCreateDetailGroup(IDetailChildrenBuilder& ChildBuilder,  TMap<FString, IDetailGroup*>& GroupMap, TSharedPtr<IPropertyHandleArray> PropertiesArray, FString& GroupName)
-	{
-		if (GroupMap.Contains(GroupName))
-		{
-			return *GroupMap[GroupName];
-		}
-
-		FString ParentGroups; 
-		FString PropertyName;
-		FText DisplayName;
-
-		if (GroupName.Split(PropertyPathDelimiter, &ParentGroups, &PropertyName, ESearchCase::IgnoreCase, ESearchDir::FromEnd))
-		{
-			IDetailGroup& ParentGroup = GetOrCreateDetailGroup(ChildBuilder, GroupMap, PropertiesArray, ParentGroups);
-			DisplayName = FText::FromString(PropertyName);
-			GroupMap.Add(GroupName, &ParentGroup.AddGroup(FName(*PropertyName), DisplayName));
-		}
-		else 
-		{
-			DisplayName = FText::FromString(GroupName);
-			GroupMap.Add(GroupName, &ChildBuilder.AddGroup(FName(*GroupName), DisplayName));
-		}
-
-		GroupMap[GroupName]->HeaderRow()
-		[
-			SNew(SHorizontalBox)
-
-	    	+SHorizontalBox::Slot()
-	    	.AutoWidth()
-	    	.VAlign(VAlign_Center)
-	    	[
-	    		SNew( SCheckBox )
-	    		.OnCheckStateChanged( this, &FRecorderPropertyMapCustomization::OnGroupCheckStateChanged, PropertiesArray, GroupName )	
-	    		.IsChecked( this, &FRecorderPropertyMapCustomization::OnGroupGetCheckState, PropertiesArray, GroupName )	
-	    	]
-
-	    	+SHorizontalBox::Slot()
-	    	.VAlign(VAlign_Center)
-	    	.AutoWidth()
-	    	.Padding(6.0, 0.0)
-	    	[
-	    		SNew(STextBlock)
-	    		.Text(DisplayName)
-	    		.Font( FAppStyle::GetFontStyle( "PropertyWindow.NormalFont" ))
-	    	]
-		];
-
-		return *GroupMap[GroupName];
-	}
-
-
-	virtual void CustomizeChildren( TSharedRef<IPropertyHandle> PropertyHandle, IDetailChildrenBuilder& ChildBuilder, IPropertyTypeCustomizationUtils& CustomizationUtils )
-	{
-		TMap<FString, IDetailGroup*> DetailGroupMap;
-
-		TSharedPtr<IPropertyHandleArray> RecordedPropertiesArrayHandle = PropertyHandle->GetChildHandle(GET_MEMBER_NAME_CHECKED(UActorRecorderPropertyMap, Properties))->AsArray();
-
-		uint32 NumProperties;
-		RecordedPropertiesArrayHandle->GetNumElements(NumProperties);
-
-
-		for(uint32 i = 0; i < NumProperties; ++i)
-		{
-			TSharedRef<IPropertyHandle> RecordedPropertyTemp = RecordedPropertiesArrayHandle->GetElement(i);
-			if (RecordedPropertyTemp->IsValidHandle())
-			{
-
-	        	TSharedPtr<IPropertyHandle> PropertyNameHandle = RecordedPropertyTemp->GetChildHandle(GET_MEMBER_NAME_CHECKED(FActorRecordedProperty, PropertyName)); 
-	        	if (PropertyNameHandle->IsValidHandle())
-	        	{
-		        	FString PropertyNameValue;
-		        	PropertyNameHandle->GetValueAsDisplayString(PropertyNameValue);
-
-		        	FString ParentGroups; 
-		        	FString PropertyName;
-		        	if (PropertyNameValue.Split(PropertyPathDelimiter, &ParentGroups, &PropertyName, ESearchCase::IgnoreCase, ESearchDir::FromEnd))
-		        	{
-		        		IDetailGroup& ParentGroup = GetOrCreateDetailGroup(ChildBuilder, DetailGroupMap, RecordedPropertiesArrayHandle, ParentGroups);
-		        		ParentGroup.AddPropertyRow(RecordedPropertyTemp);
-		        	}
-		        	else 
-		        	{
-						ChildBuilder.AddProperty(RecordedPropertyTemp);
-		        	}
-	        	}
-			}
-		}
-
-		TSharedPtr<IPropertyHandleArray> RecordedComponentsArrayHandle = PropertyHandle->GetChildHandle(GET_MEMBER_NAME_CHECKED(UActorRecorderPropertyMap, Children))->AsArray();
-
-		uint32 NumCompProperties;
-		RecordedComponentsArrayHandle->GetNumElements(NumCompProperties);
-
-		for(uint32 i = 0; i < NumCompProperties; ++i)
-		{
-			TSharedRef<IPropertyHandle> RecordedCompTemp = RecordedComponentsArrayHandle->GetElement(i);
-			if (RecordedCompTemp->IsValidHandle())
-			{
-				ChildBuilder.AddProperty(RecordedCompTemp);
-			}
-		}
-	}
-
-	void OnGroupCheckStateChanged( ECheckBoxState InNewState, TSharedPtr<IPropertyHandleArray> RecordedPropertiesArrayHandle, FString GroupName) const  
-	{
-		uint32 NumProperties;
-		RecordedPropertiesArrayHandle->GetNumElements(NumProperties);
-
-		for(uint32 i = 0; i < NumProperties; ++i)
-		{
-			TSharedRef<IPropertyHandle> RecordedPropertyTemp = RecordedPropertiesArrayHandle->GetElement(i);
-			if (RecordedPropertyTemp->IsValidHandle())
-			{
-	        	TSharedPtr<IPropertyHandle> PropertyNameHandle = RecordedPropertyTemp->GetChildHandle(GET_MEMBER_NAME_CHECKED(FActorRecordedProperty, PropertyName)); 
-	        	if (PropertyNameHandle->IsValidHandle())
-	        	{
-		        	FString PropertyNameValue;
-		        	PropertyNameHandle->GetValueAsDisplayString(PropertyNameValue);
-
-		        	if (PropertyNameValue.StartsWith(GroupName))
-		        	{
-						TSharedPtr<IPropertyHandle> EnabledRecordedPropertyTemp = RecordedPropertyTemp->GetChildHandle(GET_MEMBER_NAME_CHECKED(FActorRecordedProperty, bEnabled));
-						if (EnabledRecordedPropertyTemp->IsValidHandle())
-						{
-							EnabledRecordedPropertyTemp->SetValue( InNewState == ECheckBoxState::Checked ? true : false );
-						}
-					}
-				}
-			}
-		}
-	}
-
-	ECheckBoxState OnGroupGetCheckState( TSharedPtr<IPropertyHandleArray> RecordedPropertiesArrayHandle, FString GroupName) const 
-	{
-		bool SetFirst = false;
-		bool FinalCheckedValue = false;
-
-		uint32 NumProperties;
-		RecordedPropertiesArrayHandle->GetNumElements(NumProperties);
-
-		for(uint32 i = 0; i < NumProperties; ++i)
-		{
-			TSharedRef<IPropertyHandle> RecordedPropertyTemp = RecordedPropertiesArrayHandle->GetElement(i);
-			if (RecordedPropertyTemp->IsValidHandle())
-			{
-	        	TSharedPtr<IPropertyHandle> PropertyNameHandle = RecordedPropertyTemp->GetChildHandle(GET_MEMBER_NAME_CHECKED(FActorRecordedProperty, PropertyName)); 
-	        	if (PropertyNameHandle->IsValidHandle())
-	        	{
-		        	FString PropertyNameValue;
-		        	PropertyNameHandle->GetValueAsDisplayString(PropertyNameValue);
-
-		        	if (PropertyNameValue.StartsWith(GroupName))
-		        	{
-						TSharedPtr<IPropertyHandle> EnabledRecordedPropertyTemp = RecordedPropertyTemp->GetChildHandle(GET_MEMBER_NAME_CHECKED(FActorRecordedProperty, bEnabled));
-						if (EnabledRecordedPropertyTemp->IsValidHandle())
-						{
-							bool EnabledValue;
-							if ( EnabledRecordedPropertyTemp->GetValue(EnabledValue) == FPropertyAccess::Success )
-							{
-								if (!SetFirst)
-								{
-									FinalCheckedValue = EnabledValue;
-									SetFirst = true;
-								}
-								else if ( EnabledValue != FinalCheckedValue )
-								{
-									return ECheckBoxState::Undetermined;
-								}
-							}
-						}
-					}
-				}
-			}
-		}
-
-		return FinalCheckedValue ? ECheckBoxState::Checked : ECheckBoxState::Unchecked;
-	}
-
-	void OnCheckStateChanged( ECheckBoxState InNewState, TSharedRef<IPropertyHandle> PropertyHandle )
-	{
-		TSharedPtr<IPropertyHandleArray> RecordedPropertiesArrayHandle = PropertyHandle->GetChildHandle(GET_MEMBER_NAME_CHECKED(UActorRecorderPropertyMap, Properties))->AsArray();
-
-		uint32 NumProperties;
-		RecordedPropertiesArrayHandle->GetNumElements(NumProperties);
-
-		for(uint32 i = 0; i < NumProperties; ++i)
-		{
-			TSharedRef<IPropertyHandle> RecordedPropertyTemp = RecordedPropertiesArrayHandle->GetElement(i);
-			if (RecordedPropertyTemp->IsValidHandle())
-			{
-				TSharedPtr<IPropertyHandle> EnabledRecordedPropertyTemp = RecordedPropertyTemp->GetChildHandle(GET_MEMBER_NAME_CHECKED(FActorRecordedProperty, bEnabled));
-				if (EnabledRecordedPropertyTemp->IsValidHandle())
-				{
-					EnabledRecordedPropertyTemp->SetValue( InNewState == ECheckBoxState::Checked );
-				}
-			}
-		}
-
-		TSharedPtr<IPropertyHandleArray> RecordedComponentsArrayHandle = PropertyHandle->GetChildHandle(GET_MEMBER_NAME_CHECKED(UActorRecorderPropertyMap, Children))->AsArray();
-
-		uint32 NumCompProperties;
-		RecordedComponentsArrayHandle->GetNumElements(NumCompProperties);
-
-		for(uint32 i = 0; i < NumCompProperties; ++i)
-		{
-			TSharedPtr<IPropertyHandle> RecordedCompTemp = RecordedComponentsArrayHandle->GetElement(i);
-			if (RecordedCompTemp->IsValidHandle())
-			{
-				OnCheckStateChanged( InNewState, RecordedCompTemp.ToSharedRef());
-			}
-		}
-	}
-
-	ECheckBoxState OnGetCheckState( TSharedRef<IPropertyHandle> PropertyHandle ) const 
-	{
-		bool SetFirst = false;
-		bool FinalCheckedValue = false;
-
-		TSharedPtr<IPropertyHandleArray> RecordedPropertiesArrayHandle = PropertyHandle->GetChildHandle(GET_MEMBER_NAME_CHECKED(UActorRecorderPropertyMap, Properties))->AsArray();
-
-		uint32 NumProperties;
-		RecordedPropertiesArrayHandle->GetNumElements(NumProperties);
-
-		for(uint32 i = 0; i < NumProperties; ++i)
-		{
-			TSharedRef<IPropertyHandle> RecordedPropertyTemp = RecordedPropertiesArrayHandle->GetElement(i);
-			if (RecordedPropertyTemp->IsValidHandle())
-			{
-
-				TSharedPtr<IPropertyHandle> EnabledRecordedPropertyTemp = RecordedPropertyTemp->GetChildHandle(GET_MEMBER_NAME_CHECKED(FActorRecordedProperty, bEnabled));
-				if (EnabledRecordedPropertyTemp->IsValidHandle())
-				{
-					bool EnabledValue;
-					if ( EnabledRecordedPropertyTemp->GetValue(EnabledValue) == FPropertyAccess::Success )
-					{
-						if (!SetFirst)
-						{
-							FinalCheckedValue = EnabledValue;
-							SetFirst = true;
-						}
-						else if ( EnabledValue != FinalCheckedValue )
-						{
-							return ECheckBoxState::Undetermined;
-						}
-					}
-				}
-			}
-		}
-
-		TSharedPtr<IPropertyHandleArray> RecordedComponentsArrayHandle = PropertyHandle->GetChildHandle(GET_MEMBER_NAME_CHECKED(UActorRecorderPropertyMap, Children))->AsArray();
-
-		uint32 NumCompProperties;
-		RecordedComponentsArrayHandle->GetNumElements(NumCompProperties);
-
-		for(uint32 i = 0; i < NumCompProperties; ++i)
-		{
-			TSharedPtr<IPropertyHandle> RecordedCompTemp = RecordedComponentsArrayHandle->GetElement(i);
-			if (RecordedCompTemp->IsValidHandle())
-			{
-				ECheckBoxState CompEnabledState = OnGetCheckState( RecordedCompTemp.ToSharedRef() );
-				if (CompEnabledState == ECheckBoxState::Undetermined)
-				{
-					return ECheckBoxState::Undetermined;
-				}
-				else 
-				{
-					bool IsCompChecked = CompEnabledState == ECheckBoxState::Checked ? true : false;
-
-					if ( !SetFirst )
-					{
-						FinalCheckedValue = IsCompChecked;	
-						SetFirst = true;
-					}
-					else if ( IsCompChecked != FinalCheckedValue )
-					{
-						return ECheckBoxState::Undetermined;
-					}
-				}
-			}
-		}
-
-		return FinalCheckedValue ? ECheckBoxState::Checked : ECheckBoxState::Unchecked;
-	}
-
-};
-const FString FRecorderPropertyMapCustomization::PropertyPathDelimiter = FString(TEXT("."));
-
-class FAudioInputChannelPropertyCustomization : public IPropertyTypeCustomization
-{
-public:
-
-	virtual void CustomizeHeader(TSharedRef<IPropertyHandle> PropertyHandle, FDetailWidgetRow& HeaderRow, IPropertyTypeCustomizationUtils& ) {}
-
-	virtual void CustomizeChildren(TSharedRef<IPropertyHandle> PropertyHandle, IDetailChildrenBuilder& ChildBuilder, IPropertyTypeCustomizationUtils& CustomizationUtils)
-	{
-		InputDeviceChannelHandle = PropertyHandle->GetChildHandle(GET_MEMBER_NAME_CHECKED(FAudioInputDeviceChannelProperty, AudioInputDeviceChannel));
-
-		if (InputDeviceChannelHandle.IsValid())
-		{
-			IDetailPropertyRow& InputDeviceChannelPropertyRow = ChildBuilder.AddProperty(PropertyHandle);
-			InputDeviceChannelPropertyRow.CustomWidget()
-			.NameContent()
-			[
-				InputDeviceChannelHandle->CreatePropertyNameWidget()
-			]
-			.ValueContent()
-			[
-				MakeInputChannelSelectorWidget()
-			];
-		}
-
-		// Register delegate so UI can update when the audio input device changes
-		if (UTakeRecorderAudioInputSettings* AudioInputSettings = TakeRecorderAudioSettingsUtils::GetTakeRecorderAudioInputSettings())
-		{
-			AudioInputSettings->GetOnAudioInputDeviceChanged().Add(FSimpleDelegate::CreateSP(this, &FAudioInputChannelPropertyCustomization::RebuildInputChannelArray));
-		}
-	}
-
-private:
-
-	void BuildInputChannelArray()
-	{
-		if (UTakeRecorderAudioInputSettings* AudioInputSettings = TakeRecorderAudioSettingsUtils::GetTakeRecorderAudioInputSettings())
-		{
-			int32 DeviceChannelCount = AudioInputSettings->GetDeviceChannelCount();
-
-			if (DeviceChannelCount > 0)
-			{
-				for (int32 ChannelIndex = 1; ChannelIndex <= DeviceChannelCount; ++ChannelIndex)
-				{
-					TSharedPtr<int32> ChannelIndexPtr = MakeShared<int32>(ChannelIndex);
-					InputDeviceChannelArray.Add(ChannelIndexPtr);
-				}
-			}
-
-			int32 CurrentValue;
-			InputDeviceChannelHandle->GetValue(CurrentValue);
-			if (CurrentValue > DeviceChannelCount)
-			{
-				InputDeviceChannelHandle->SetValue(0);
-			}
-		}
-	}
-
-	void RebuildInputChannelArray()
-	{
-		InputDeviceChannelArray.Empty();
-		BuildInputChannelArray();
-		ChannelComboBox->RefreshOptions();
-
-		UTakeRecorderAudioInputSettings* AudioInputSettings = TakeRecorderAudioSettingsUtils::GetTakeRecorderAudioInputSettings();
-
-		if (AudioInputSettings && InputDeviceChannelHandle.IsValid())
-		{
-			int32 CurrentValue;
-			InputDeviceChannelHandle->GetValue(CurrentValue);
-			int32 DeviceChannelCount = AudioInputSettings->GetDeviceChannelCount();
-
-			if (CurrentValue < 1 || CurrentValue > DeviceChannelCount)
-			{
-				InputChannelTitleBlock->SetText(FText());
-			}
-		}
-	}
-
-	TSharedRef<SWidget> MakeInputChannelSelectorWidget()
-	{
-		BuildInputChannelArray();
-
-		// Combo box component:
-		SAssignNew(ChannelComboBox, SComboBox<TSharedPtr<int32>>)
-			.OptionsSource(&InputDeviceChannelArray)
-			.OnGenerateWidget_Lambda([](TSharedPtr<int32> InValue)
-			{
-				return SNew(STextBlock)
-					.Text(FText::AsNumber(*InValue))
-					.Font(FAppStyle::Get().GetFontStyle("SmallFont"));
-			})
-			.OnSelectionChanged_Lambda([this](TSharedPtr<int32> InSelection, ESelectInfo::Type InSelectInfo)
-			{
-				if (InSelection.IsValid())
-				{
-					InputDeviceChannelHandle->SetValue(*InSelection);
-					InputChannelTitleBlock->SetText(FText::AsNumber(*InSelection));
-				}
-			})
-			[
-				SAssignNew(InputChannelTitleBlock, STextBlock)
-				.Text_Lambda([this]()
-				{
-					int32 ChannelNumber;
-					InputDeviceChannelHandle->GetValue(ChannelNumber);
-
-					if (ChannelNumber > 0)
-					{
-						return FText::AsNumber(ChannelNumber);
-					}
-					return FText();
-				})
-				.Font(FAppStyle::Get().GetFontStyle("SmallFont"))
-			];
-
-		return ChannelComboBox.ToSharedRef();
-	}
-
-	TSharedPtr<IPropertyHandle> InputDeviceChannelHandle;
-	TArray<TSharedPtr<int32>> InputDeviceChannelArray;
-	TSharedPtr<STextBlock> InputChannelTitleBlock;
-	TSharedPtr<SComboBox<TSharedPtr<int32>>> ChannelComboBox;
-};
-
-class FRecorderSourceObjectCustomization : public IDetailCustomization
-{
-	virtual void CustomizeDetails(IDetailLayoutBuilder& DetailBuilder) override
-	{
-		FText NewTitle = ComputeTitle(DetailBuilder.GetDetailsViewSharedPtr());
-		if (!NewTitle.IsEmpty())
-		{
-			// Edit the category and add *all* properties for the object to it
-			IDetailCategoryBuilder& CategoryBuilder = DetailBuilder.EditCategory("CustomCategory", NewTitle);
-
-			UClass* BaseClass = DetailBuilder.GetBaseClass();
-			while (BaseClass)
-			{
-				for (FProperty* Property : TFieldRange<FProperty>(BaseClass, EFieldIteratorFlags::ExcludeSuper))
-				{
-					CategoryBuilder.AddProperty(Property->GetFName(), BaseClass);
-				}
-
-				BaseClass = BaseClass->GetSuperClass();
-			}
-		}
-	}
-
-	FText ComputeTitle(const TSharedPtr<const IDetailsView>& DetailsView) const
-	{
-		if (!DetailsView)
-		{
-			return FText();
-		}
-
-		// Compute the title for all the sources that this details panel is editing
-		static const FName CategoryName = "Category";
-
-		const TArray<TWeakObjectPtr<UObject>>& SelectedObjects = DetailsView->GetSelectedObjects();
-		UObject* FirstObject = SelectedObjects.Num() > 0 ? SelectedObjects[0].Get() : nullptr;
-
-		if (FirstObject)
-		{
-			if (SelectedObjects.Num() == 1)
-			{
-				UTakeRecorderSource* Source = Cast<UTakeRecorderSource>(FirstObject);
-				return Source ? Source->GetDisplayText() : FText::FromString(FirstObject->GetName());
-			}
-			else if (SelectedObjects.Num() > 1)
-			{
-				const FString& Category = FirstObject->GetClass()->GetMetaData(CategoryName);
-				return FText::Format(LOCTEXT("CategoryFormatString", "{0} ({1})"), FText::FromString(Category), SelectedObjects.Num());
-			}
-		}
-
-		return FText();
-	}
-};
-
-
 void SLevelSequenceTakeEditor::AddDetails(const TPair<const UClass*, TArray<UObject*> >& Pair, TArray<FObjectKey>& PreviousClasses)
 {
 	FPropertyEditorModule& PropertyEditorModule = FModuleManager::Get().LoadModuleChecked<FPropertyEditorModule>("PropertyEditor");
@@ -910,6 +300,7 @@ void SLevelSequenceTakeEditor::AddDetails(const TPair<const UClass*, TArray<UObj
 	}
 	else
 	{
+		using namespace UE::TakeRecorder;
 		TSharedRef<IDetailsView> Details = PropertyEditorModule.CreateDetailView(DetailsViewArgs);
 
 		// Register the custom property layout for all object types to rename the category to the object type
diff --git a/Engine/Plugins/VirtualProduction/Takes/Source/TakeRecorder/Private/Widgets/TakeRecorderAudioSettingsCustomization.cpp b/Engine/Plugins/VirtualProduction/Takes/Source/TakeRecorder/Private/Widgets/TakeRecorderAudioSettingsCustomization.cpp
deleted file mode 100644
diff --git a/Engine/Plugins/VirtualProduction/Takes/Source/TakeRecorder/Private/Widgets/TakeRecorderAudioSettingsCustomization.h b/Engine/Plugins/VirtualProduction/Takes/Source/TakeRecorder/Private/Widgets/TakeRecorderAudioSettingsCustomization.h
deleted file mode 100644
