TakeRecorder: Engine subsystem is now responsible for managing recordings, removing the requirement of leaving the Take Recorder panel open.

diff --git a/Engine/Plugins/VirtualProduction/Takes/Source/TakeRecorder/Private/Recorder/TakeRecorder.cpp b/Engine/Plugins/VirtualProduction/Takes/Source/TakeRecorder/Private/Recorder/TakeRecorder.cpp
--- a/Engine/Plugins/VirtualProduction/Takes/Source/TakeRecorder/Private/Recorder/TakeRecorder.cpp
+++ b/Engine/Plugins/VirtualProduction/Takes/Source/TakeRecorder/Private/Recorder/TakeRecorder.cpp
@@ -394,6 +394,11 @@ FTakeRecorderParameterOverride& UTakeRecorder::TakeInitializeParameterOverride()
 	return Overrides;
 }
 
+void UTakeRecorder::SetCountdown(float InSeconds)
+{
+	CountdownSeconds = InSeconds;
+}
+
 bool UTakeRecorder::SetActiveRecorder(UTakeRecorder* NewActiveRecorder)
 {
 	if (GetCurrentRecorder().IsValid())
@@ -465,20 +470,23 @@ bool UTakeRecorder::Initialize ( ULevelSequence* LevelSequenceBase, UTakeRecorde
 	UTakeRecorderBlueprintLibrary::OnTakeRecorderPreInitialize();
 	if (UTakeRecorderSubsystem* TakeRecorderSubsystem = GEngine->GetEngineSubsystem<UTakeRecorderSubsystem>())
 	{
+		// The editor subsystem is normally responsible broadcasting their delegates, but they don't have a reliable way of binding to a
+		// pre-initialize event.
+		TakeRecorderSubsystem->GetOnRecordingPreInitializedEvent().Broadcast(this);
 		TakeRecorderSubsystem->TakeRecorderPreInitialize.Broadcast();
 	}
 
 	FTakeRecorderParameters FinalParameters = TakeInitHelper::AccumulateParamsOverride(InParameters);
 	if (FinalParameters.TakeRecorderMode == ETakeRecorderMode::RecordNewSequence)
 	{
-		if (!CreateDestinationAsset(*FinalParameters.Project.GetTakeAssetPath(), LevelSequenceBase, Sources, MetaData, OutError))
+		if (!CreateDestinationAsset(*FinalParameters.Project.GetTakeAssetPath(), LevelSequenceBase, Sources, MetaData, OutError, FinalParameters.bOpenSequencer))
 		{
 			return false;
 		}
 	}
 	else
 	{
-		if (!SetupDestinationAsset(FinalParameters, LevelSequenceBase, Sources, MetaData, OutError))
+		if (!SetupDestinationAsset(FinalParameters, LevelSequenceBase, Sources, MetaData, OutError, FinalParameters.bOpenSequencer))
 		{
 			return false;
 		}
@@ -598,7 +606,7 @@ void UTakeRecorder::DiscoverSourceWorld()
 	}
 }
 
-bool UTakeRecorder::CreateDestinationAsset(const TCHAR* AssetPathFormat, ULevelSequence* LevelSequenceBase, UTakeRecorderSources* Sources, UTakeMetaData* MetaData, FText* OutError)
+bool UTakeRecorder::CreateDestinationAsset(const TCHAR* AssetPathFormat, ULevelSequence* LevelSequenceBase, UTakeRecorderSources* Sources, UTakeMetaData* MetaData, FText* OutError, bool bOpenSequencer)
 {
 	check(LevelSequenceBase && Sources && MetaData);
 
@@ -634,18 +642,26 @@ bool UTakeRecorder::CreateDestinationAsset(const TCHAR* AssetPathFormat, ULevelS
 	SequenceAsset->MarkPackageDirty();
 	FAssetRegistryModule::AssetCreated(SequenceAsset);
 
-	WeakSequencer = TakesUtils::OpenSequencer(SequenceAsset, OutError);
-	return WeakSequencer.IsValid();
+	bool bSuccess = true;
+	if (bOpenSequencer)
+	{
+		WeakSequencer = TakesUtils::OpenSequencer(SequenceAsset, OutError);
+		bSuccess = WeakSequencer.IsValid();
+	}
+	return bSuccess;
 }
 
-bool UTakeRecorder::SetupDestinationAsset(const FTakeRecorderParameters& InParameters, ULevelSequence* LevelSequenceBase, UTakeRecorderSources* Sources, UTakeMetaData* MetaData, FText* OutError)
+bool UTakeRecorder::SetupDestinationAsset(const FTakeRecorderParameters& InParameters, ULevelSequence* LevelSequenceBase, UTakeRecorderSources* Sources, UTakeMetaData* MetaData, FText* OutError, bool bOpenSequencer)
 {
 	check(LevelSequenceBase && Sources && MetaData);
 
-	WeakSequencer = TakesUtils::OpenSequencer(LevelSequenceBase, OutError);
-	if (!WeakSequencer.IsValid())
+	if (bOpenSequencer)
 	{
-		return false;
+		WeakSequencer = TakesUtils::OpenSequencer(LevelSequenceBase, OutError);
+		if (!WeakSequencer.IsValid())
+		{
+			return false;
+		}
 	}
 
 	// The SequenceAsset is either LevelSequenceBase or the currently focused sequence
@@ -998,10 +1014,6 @@ void UTakeRecorder::Start()
 	OnRecordingStartedEvent.Broadcast(this);
 
 	UTakeRecorderBlueprintLibrary::OnTakeRecorderStarted();
-	if (UTakeRecorderSubsystem* TakeRecorderSubsystem = GEngine->GetEngineSubsystem<UTakeRecorderSubsystem>())
-	{
-		TakeRecorderSubsystem->TakeRecorderStarted.Broadcast();
-	}
 
 	if (Sequencer.IsValid())
 	{
@@ -1089,10 +1101,6 @@ void UTakeRecorder::StopInternal(const bool bCancelled)
 		OnRecordingStoppedEvent.Broadcast(this);
 
 		UTakeRecorderBlueprintLibrary::OnTakeRecorderStopped();
-		if (UTakeRecorderSubsystem* TakeRecorderSubsystem = GEngine->GetEngineSubsystem<UTakeRecorderSubsystem>())
-		{
-			TakeRecorderSubsystem->TakeRecorderStopped.Broadcast();
-		}
 
 		if (MovieScene)
 		{
@@ -1207,19 +1215,11 @@ void UTakeRecorder::StopInternal(const bool bCancelled)
 		{
 			OnRecordingFinishedEvent.Broadcast(this);
 			UTakeRecorderBlueprintLibrary::OnTakeRecorderFinished(SequenceAsset);
-			if (UTakeRecorderSubsystem* TakeRecorderSubsystem = GEngine->GetEngineSubsystem<UTakeRecorderSubsystem>())
-			{
-				TakeRecorderSubsystem->TakeRecorderFinished.Broadcast(SequenceAsset);
-			}
 		}
 		else
 		{
 			OnRecordingCancelledEvent.Broadcast(this);
 			UTakeRecorderBlueprintLibrary::OnTakeRecorderCancelled();
-			if (UTakeRecorderSubsystem* TakeRecorderSubsystem = GEngine->GetEngineSubsystem<UTakeRecorderSubsystem>())
-			{
-				TakeRecorderSubsystem->TakeRecorderCancelled.Broadcast();
-			}
 		}
 	}
 
diff --git a/Engine/Plugins/VirtualProduction/Takes/Source/TakeRecorder/Private/Recorder/TakeRecorderParameters.cpp b/Engine/Plugins/VirtualProduction/Takes/Source/TakeRecorder/Private/Recorder/TakeRecorderParameters.cpp
--- a/Engine/Plugins/VirtualProduction/Takes/Source/TakeRecorder/Private/Recorder/TakeRecorderParameters.cpp
+++ b/Engine/Plugins/VirtualProduction/Takes/Source/TakeRecorder/Private/Recorder/TakeRecorderParameters.cpp
@@ -41,3 +41,20 @@ FTakeRecorderParameters::FTakeRecorderParameters()
 	: TakeRecorderMode(ETakeRecorderMode::RecordNewSequence)
 {}
 
+bool FTakeRecorderSequenceParameters::IsValid() const
+{
+	return *this != FTakeRecorderSequenceParameters();
+}
+
+bool FTakeRecorderSequenceParameters::operator==(const FTakeRecorderSequenceParameters& Other) const
+{
+	return BasePreset == Other.BasePreset
+		&& BaseSequence == Other.BaseSequence
+		&& RecordIntoSequence == Other.RecordIntoSequence
+		&& SequenceToView == Other.SequenceToView;
+}
+
+bool FTakeRecorderSequenceParameters::operator!=(const FTakeRecorderSequenceParameters& Other) const
+{
+	return !(*this == Other);
+}
diff --git a/Engine/Plugins/VirtualProduction/Takes/Source/TakeRecorder/Private/Recorder/TakeRecorderSubsystem.cpp b/Engine/Plugins/VirtualProduction/Takes/Source/TakeRecorder/Private/Recorder/TakeRecorderSubsystem.cpp
--- a/Engine/Plugins/VirtualProduction/Takes/Source/TakeRecorder/Private/Recorder/TakeRecorderSubsystem.cpp
+++ b/Engine/Plugins/VirtualProduction/Takes/Source/TakeRecorder/Private/Recorder/TakeRecorderSubsystem.cpp
@@ -2,6 +2,1033 @@
 
 #include "Recorder/TakeRecorderSubsystem.h"
 
-#define LOCTEXT_NAMESPACE "TakeRecorder"
+#include "AssetRegistry/AssetRegistryModule.h"
+#include "Dialog/SMessageDialog.h"
+#include "Framework/Notifications/NotificationManager.h"
+#include "ILevelSequenceEditorToolkit.h"
+#include "ITakeRecorderModule.h"
+#include "ITakeRecorderNamingTokensModule.h"
+#include "LevelSequence.h"
+#include "MovieScene.h"
+#include "NamingTokens.h"
+#include "NamingTokensEngineSubsystem.h"
+#include "Recorder/TakeRecorder.h"
+#include "Recorder/TakeRecorderBlueprintLibrary.h"
+#include "Recorder/TakeRecorderParameters.h"
+#include "ScopedTransaction.h"
+#include "TakeMetaData.h"
+#include "TakeRecorderModule.h"
+#include "TakeRecorderSettings.h"
+#include "TakeRecorderSource.h"
+#include "TakeRecorderSources.h"
+#include "TakesCoreBlueprintLibrary.h"
+#include "Widgets/Notifications/SNotificationList.h"
+#include "Widgets/TakeRecorderNamingTokensData.h"
+
+#define LOCTEXT_NAMESPACE "TakeRecorderSubsystem"
+
+void UTakeRecorderSubsystem::Initialize(FSubsystemCollectionBase& Collection)
+{
+	Super::Initialize(Collection);
+
+	const FAssetRegistryModule& AssetRegistryModule = FModuleManager::LoadModuleChecked<FAssetRegistryModule>(TEXT("AssetRegistry"));
+	OnAssetRegistryFilesLoadedHandle = AssetRegistryModule.Get().OnFilesLoaded().AddUObject(this, &UTakeRecorderSubsystem::OnAssetRegistryFilesLoaded);
+	OnRecordingInitializedHandle = UTakeRecorder::OnRecordingInitialized().AddUObject(this, &UTakeRecorderSubsystem::OnRecordingInitialized);
+
+	UTakeRecorderSources::OnSourceAdded().AddUObject(this, &UTakeRecorderSubsystem::OnSourceAdded);
+	UTakeRecorderSources::OnSourceRemoved().AddUObject(this, &UTakeRecorderSubsystem::OnSourceRemoved);
+
+	NamingTokensData = NewObject<UTakeRecorderNamingTokensData>();
+	NamingTokensData->SetFlags(RF_Transactional);
+
+	BindToNamingTokenEvents();
+}
+
+void UTakeRecorderSubsystem::Deinitialize()
+{
+	Super::Deinitialize();
+	UTakeRecorder::OnRecordingInitialized().Remove(OnRecordingInitializedHandle);
+
+	UTakeRecorderSources::OnSourceAdded().RemoveAll(this);
+	UTakeRecorderSources::OnSourceRemoved().RemoveAll(this);
+	
+	if (const FAssetRegistryModule* AssetRegistryModule = FModuleManager::GetModulePtr<FAssetRegistryModule>("AssetRegistry"))
+	{
+		if (IAssetRegistry* AssetRegistry = AssetRegistryModule->TryGet())
+		{
+			AssetRegistry->OnFilesLoaded().Remove(OnAssetRegistryFilesLoadedHandle);
+		}
+	}
+
+	if (ITakeRecorderModule* TakeRecorderModule = FModuleManager::Get().GetModulePtr<ITakeRecorderModule>("TakeRecorder"))
+	{
+		TakeRecorderModule->GetLastLevelSequenceProvider().Unbind();
+	}
+
+	UnbindNamingTokensEvents();
+}
+
+void UTakeRecorderSubsystem::Tick(float DeltaTime)
+{
+	CacheMetaData();
+}
+
+TStatId UTakeRecorderSubsystem::GetStatId() const
+{
+	RETURN_QUICK_DECLARE_CYCLE_STAT(UTakeRecorderSubsystem, STATGROUP_Tickables);
+}
+
+void UTakeRecorderSubsystem::SetTargetSequence(const FTakeRecorderSequenceParameters& InData)
+{
+	// Don't initialize again if we already have the same data present.
+	if (TargetSequenceData.IsValid() && TargetSequenceData == InData)
+	{
+		return;
+	}
+	TargetSequenceData = InData;
+	
+	ITakeRecorderModule& TakeRecorderModule = FModuleManager::Get().LoadModuleChecked<ITakeRecorderModule>("TakeRecorder");
+	TakeRecorderModule.GetLastLevelSequenceProvider().BindUObject(this, &UTakeRecorderSubsystem::SetLastLevelSequence);
+
+	bAutoApplyTakeNumber = true;
+	
+	CacheMetaData();
+	
+	if (TakeMetaData && !TakeMetaData->IsLocked())
+	{
+		const int32 NextTakeNumber = GetNextTakeNumber(TakeMetaData->GetSlate());
+		if (NextTakeNumber != TakeMetaData->GetTakeNumber())
+		{
+			TakeMetaData->SetTakeNumber(NextTakeNumber);
+		}
+	}
+	
+	// If a recording is currently underway, initialize to that now
+	if (UTakeRecorder* ActiveRecorder = UTakeRecorder::GetActiveRecorder())
+	{
+		RecordingLevelSequence = ActiveRecorder->GetSequence();
+
+		ActiveRecorder->OnRecordingStopped().Remove(OnRecordingStoppedHandle);
+		ActiveRecorder->OnRecordingFinished().Remove(OnRecordingFinishedHandle);
+		ActiveRecorder->OnRecordingCancelled().Remove(OnRecordingCancelledHandle);
+		OnRecordingStoppedHandle = ActiveRecorder->OnRecordingStopped().AddUObject(this, &UTakeRecorderSubsystem::OnRecordingStopped);
+		OnRecordingFinishedHandle = ActiveRecorder->OnRecordingFinished().AddUObject(this, &UTakeRecorderSubsystem::OnRecordingFinished);
+		OnRecordingCancelledHandle = ActiveRecorder->OnRecordingCancelled().AddUObject(this, &UTakeRecorderSubsystem::OnRecordingCancelled);
+	}
+	else
+	{
+		RecordingLevelSequence = nullptr;
+	}
+	
+	TransientPreset = AllocateTransientPreset();
+	
+	// Copy the base preset into the transient preset if it was provided.
+	// We do this first so that anything that asks for its Level Sequence
+	// on construction gets the right one
+	if (InData.BasePreset)
+	{
+		TransientPreset->CopyFrom(InData.BasePreset);
+	}
+	else if (InData.BaseSequence)
+	{
+		TransientPreset->CopyFrom(InData.BaseSequence);
+
+		ULevelSequence* LevelSequence = TransientPreset->GetLevelSequence();
+
+		UTakeRecorderSources* BaseSources = InData.BaseSequence->FindMetaData<UTakeRecorderSources>();
+		if (BaseSources && LevelSequence)
+		{
+			LevelSequence->CopyMetaData(BaseSources);
+		}
+
+		if (LevelSequence)
+		{
+			LevelSequence->GetMovieScene()->SetReadOnly(false);
+		}
+
+		if (UTakeMetaData*  TakeMetaDataLevelSequence = LevelSequence ? LevelSequence->FindMetaData<UTakeMetaData>() : nullptr)
+		{
+			TakeMetaDataLevelSequence->Unlock();
+			TakeMetaDataLevelSequence->SetTimestamp(FDateTime(0));
+		}
+	}
+	else if (InData.RecordIntoSequence)
+	{
+		SetRecordIntoLevelSequence(InData.RecordIntoSequence);
+	}
+	else if (InData.SequenceToView)
+	{
+		SuppliedLevelSequence = InData.SequenceToView;
+	}
+
+	/** Clear the dirty flag since the preset was just initialized. */
+	TransientPreset->GetOutermost()->SetDirtyFlag(false);
+}
+
+void UTakeRecorderSubsystem::SetRecordIntoLevelSequence(ULevelSequence* LevelSequence)
+{
+	SuppliedLevelSequence = nullptr;
+	RecordIntoLevelSequence = LevelSequence;
+
+	if (RecordIntoLevelSequence)
+	{
+		RecordIntoLevelSequence->GetMovieScene()->SetReadOnly(false);
+	}
+
+	if (UTakeMetaData* LocalTakeMetaData = RecordIntoLevelSequence ? RecordIntoLevelSequence->FindOrAddMetaData<UTakeMetaData>() : nullptr)
+	{
+		// Set up take metadata to match this level sequence's info, ie. match the frame rate, use the level sequence name as the slate
+		LocalTakeMetaData->Unlock();
+		LocalTakeMetaData->SetTimestamp(FDateTime(0));
+		LocalTakeMetaData->SetSlate(LevelSequence->GetName());
+		LocalTakeMetaData->SetTakeNumber(0);
+		LocalTakeMetaData->SetFrameRate(LevelSequence->GetMovieScene()->GetDisplayRate());
+		LocalTakeMetaData->SetFrameRateFromTimecode(false);
+	}
+}
+
+bool UTakeRecorderSubsystem::CanReviewLastRecording() const
+{
+	if (GetLastRecordedLevelSequence() == nullptr)
+	{
+		return false;
+	}
+	ITakeRecorderModule& TakeRecorderModule = FModuleManager::Get().LoadModuleChecked<ITakeRecorderModule>("TakeRecorder");
+	FCanReviewLastRecordedLevelSequence& CanReview = TakeRecorderModule.GetCanReviewLastRecordedLevelSequenceDelegate();
+	if (CanReview.IsBound())
+	{
+		return CanReview.Execute();
+	}
+	
+	return true;
+}
+
+bool UTakeRecorderSubsystem::ReviewLastRecording()
+{
+	if (LastRecordedLevelSequence)
+	{
+		SuppliedLevelSequence = LastRecordedLevelSequence;
+		LastRecordedLevelSequence = nullptr;
+		return true;
+	}
+
+	return false;
+}
+
+bool UTakeRecorderSubsystem::StartRecording(bool bOpenSequencer, bool bShowErrorMessage)
+{
+	static bool bStartedRecording = false;
+
+	if (bStartedRecording)
+	{
+		return false;
+	}
+
+	if (!GIsEditor)
+	{
+		// Sequencer requires an editor attached.
+		bOpenSequencer = false;
+	}
+
+	TGuardValue<bool> ReentrantGuard(bStartedRecording, true);
+
+	ULevelSequence*       LevelSequence = GetLevelSequence();
+	UTakeRecorderSources* Sources = LevelSequence ? LevelSequence->FindMetaData<UTakeRecorderSources>() : nullptr;
+	
+	if (LevelSequence && Sources)
+	{
+		FTakeRecorderParameters Parameters;
+		Parameters.User    = GetDefault<UTakeRecorderUserSettings>()->Settings;
+		Parameters.Project = GetDefault<UTakeRecorderProjectSettings>()->Settings;
+		Parameters.TakeRecorderMode = GetTakeRecorderMode();
+		Parameters.StartFrame = LevelSequence->GetMovieScene()->GetPlaybackRange().GetLowerBoundValue();
+		Parameters.bOpenSequencer = bOpenSequencer;
+
+		FText ErrorText = LOCTEXT("UnknownError", "An unknown error occurred when trying to start recording");
+
+		if (GIsEditor)
+		{
+			IAssetEditorInstance* AssetEditor = GEditor->GetEditorSubsystem<UAssetEditorSubsystem>()->FindEditorForAsset(LevelSequence, false);
+			ILevelSequenceEditorToolkit* LevelSequenceEditor = static_cast<ILevelSequenceEditorToolkit*>(AssetEditor);
+
+			if (LevelSequenceEditor && LevelSequenceEditor->GetSequencer())
+			{
+				// If not resetting the playhead, store the current time as the start frame for recording. 
+				// This will ultimately be the start of the playback range and the recording will begin from that time.
+				if (!Parameters.User.bResetPlayhead)
+				{
+					Parameters.StartFrame = LevelSequenceEditor->GetSequencer()->GetLocalTime().Time.FrameNumber;
+				}
+			}
+		}
+
+		UTakeRecorder* NewRecorder = NewObject<UTakeRecorder>(GetTransientPackage(), NAME_None, RF_Transient);
+
+		if (!NewRecorder->Initialize(LevelSequence, Sources, TakeMetaData, Parameters, &ErrorText))
+		{
+			if (bShowErrorMessage && ensure(!ErrorText.IsEmpty()))
+			{
+				FNotificationInfo Info(ErrorText);
+				Info.ExpireDuration = 5.0f;
+				FSlateNotificationManager::Get().AddNotification(Info)->SetCompletionState(SNotificationItem::CS_Fail);
+			}
+			return false;
+		}
+
+		OnRecordingStarted(NewRecorder);
+		
+		return true;
+	}
+	
+	return false;
+}
+
+void UTakeRecorderSubsystem::StopRecording()
+{
+	if (UTakeRecorder* CurrentRecording = UTakeRecorder::GetActiveRecorder())
+	{
+		CurrentRecording->Stop();
+	}
+}
+
+void UTakeRecorderSubsystem::CancelRecording(bool bShowConfirmMessage)
+{
+	if (UTakeRecorder* CurrentRecording = UTakeRecorder::GetActiveRecorder())
+	{
+		auto DoCancel = [](const TWeakObjectPtr<UTakeRecorder>& WeakRecording)
+		{
+			if (WeakRecording.IsValid())
+			{
+				WeakRecording.Get()->Cancel(); 
+			}
+		};
+
+		if (bShowConfirmMessage)
+		{
+			const TSharedRef<SMessageDialog> ConfirmDialog = SNew(SMessageDialog)
+				.Title(FText(LOCTEXT("ConfirmCancelRecordingTitle", "Cancel Recording?")))
+				.Message(LOCTEXT("ConfirmCancelRecording", "Are you sure you want to cancel the current recording?"))
+				.Buttons({
+					SCustomDialog::FButton(LOCTEXT("Yes", "Yes"), FSimpleDelegate::CreateLambda([DoCancel, WeakRecording = TWeakObjectPtr<UTakeRecorder>(CurrentRecording)]()
+						{ 
+							DoCancel(WeakRecording);
+						})),
+					SCustomDialog::FButton(LOCTEXT("No", "No"))
+				});
+
+			// Non modal so that the recording continues to update
+			ConfirmDialog->Show();
+		}
+		else
+		{
+			DoCancel(TWeakObjectPtr<UTakeRecorder>(CurrentRecording));
+		}
+	}
+}
+
+void UTakeRecorderSubsystem::ResetToPendingTake()
+{
+	if (IsReviewing())
+	{
+		LastRecordedLevelSequence = SuppliedLevelSequence;
+	}
+
+	SuppliedLevelSequence = nullptr;
+	RecordIntoLevelSequence = nullptr;
+	
+	TransientPreset = AllocateTransientPreset();
+}
+
+void UTakeRecorderSubsystem::ClearPendingTake()
+{
+	if (IsReviewing())
+	{
+		LastRecordedLevelSequence = SuppliedLevelSequence;
+	}
+
+	UTakeRecorderSources* BaseSources = nullptr;
+
+	if (const ULevelSequence* CurrentLevelSequence = GetLevelSequence())
+	{
+		BaseSources = CurrentLevelSequence->FindMetaData<UTakeRecorderSources>();
+	}
+
+	SuppliedLevelSequence = nullptr;
+	RecordIntoLevelSequence = nullptr;
+
+	FScopedTransaction Transaction(LOCTEXT("ClearPendingTake_Transaction", "Clear Pending Take"));
+
+	TransientPreset->Modify();
+	TransientPreset->CreateLevelSequence();
+
+	ULevelSequence* LevelSequence = TransientPreset->GetLevelSequence();
+	if (LevelSequence && BaseSources)
+	{
+		LevelSequence->CopyMetaData(BaseSources);
+	}
+}
+
+UTakePreset* UTakeRecorderSubsystem::GetPendingTake() const
+{
+	const ITakeRecorderModule& TakeRecorderModule = FModuleManager::LoadModuleChecked<ITakeRecorderModule>("TakeRecorder");
+	return TakeRecorderModule.GetPendingTake();
+}
+
+void UTakeRecorderSubsystem::RevertChanges()
+{
+	UTakePreset* PresetOrigin = GetTakeMetaData()->GetPresetOrigin();
+
+	FScopedTransaction Transaction(LOCTEXT("RevertChanges_Transaction", "Revert Changes"));
+
+	TransientPreset->Modify();
+	TransientPreset->CopyFrom(PresetOrigin);
+	TransientPreset->GetOutermost()->SetDirtyFlag(false);
+}
+
+UTakeRecorderSource* UTakeRecorderSubsystem::AddSource(const TSubclassOf<UTakeRecorderSource> InSourceClass)
+{
+	UTakeRecorderSources* Sources = GetSources();
+
+	if (*InSourceClass && Sources)
+	{
+		FScopedTransaction Transaction(FText::Format(LOCTEXT("AddNewSource", "Add New {0} Source"), InSourceClass->GetDisplayNameText()));
+		Sources->Modify();
+
+		return Sources->AddSource(InSourceClass);
+	}
+
+	return nullptr;
+}
+
+void UTakeRecorderSubsystem::RemoveSource(UTakeRecorderSource* InSource)
+{
+	UTakeRecorderSources* Sources = GetSources();
+	if (Sources && InSource)
+	{
+		Sources->RemoveSource(InSource);
+	}
+}
+
+void UTakeRecorderSubsystem::ClearSources()
+{
+	if (UTakeRecorderSources* Sources = GetSources())
+	{
+		for (UTakeRecorderSource* Source : Sources->GetSourcesCopy())
+		{
+			Sources->RemoveSource(Source);
+		}
+	}
+}
+
+UTakeRecorderSources* UTakeRecorderSubsystem::GetSources() const
+{
+	ULevelSequence* LevelSequence = GetLevelSequence();
+	UTakeRecorderSources* Sources = LevelSequence ? LevelSequence->FindOrAddMetaData<UTakeRecorderSources>() : nullptr;
+	return Sources;
+}
+
+TArrayView<UTakeRecorderSource* const> UTakeRecorderSubsystem::GetAllSources() const
+{
+	if (const UTakeRecorderSources* Sources = GetSources())
+	{
+		return Sources->GetSources();
+	}
+
+	return {};
+}
+
+TArray<UTakeRecorderSource*> UTakeRecorderSubsystem::GetAllSourcesCopy() const
+{
+	return TArray<UTakeRecorderSource*>(GetAllSources());
+}
+
+UTakeRecorderSource* UTakeRecorderSubsystem::GetSourceByClass(const TSubclassOf<UTakeRecorderSource> InSourceClass) const
+{
+	const TArrayView<UTakeRecorderSource* const> Sources = GetAllSources();
+	for (UTakeRecorderSource* Source : Sources)
+	{
+		if (Source->IsA(InSourceClass))
+		{
+			return Source;
+		}
+	}
+
+	return nullptr;
+}
+
+void UTakeRecorderSubsystem::AddSourceForActor(AActor* InActor, bool bReduceKeys, bool bShowProgress)
+{
+	if (UTakeRecorderSources* Sources = GetSources())
+	{
+		FTakeRecorderModule& TakeRecorderModule = FModuleManager::GetModuleChecked<FTakeRecorderModule>("TakeRecorder");
+		TakeRecorderModule.GetSourcesExtensionData().OnAddActorSource.ExecuteIfBound(Sources, { InActor }, bReduceKeys, bShowProgress);
+	}
+}
+
+void UTakeRecorderSubsystem::RemoveActorFromSources(AActor* InActor)
+{
+	if (UTakeRecorderSources* Sources = GetSources())
+	{
+		FTakeRecorderModule& TakeRecorderModule = FModuleManager::GetModuleChecked<FTakeRecorderModule>("TakeRecorder");
+		TakeRecorderModule.GetSourcesExtensionData().OnRemoveActorSource.ExecuteIfBound(Sources, { InActor });
+	}
+}
+
+AActor* UTakeRecorderSubsystem::GetSourceActor(UTakeRecorderSource* InSource) const
+{
+	FTakeRecorderModule& TakeRecorderModule = FModuleManager::GetModuleChecked<FTakeRecorderModule>("TakeRecorder");
+	if (TakeRecorderModule.GetSourcesExtensionData().OnGetSourceActor.IsBound())
+	{
+		return TakeRecorderModule.GetSourcesExtensionData().OnGetSourceActor.Execute(InSource);
+	}
+
+	return nullptr;
+}
+
+ETakeRecorderState UTakeRecorderSubsystem::GetState() const
+{
+	if (const UTakeRecorder* Recorder = UTakeRecorder::GetActiveRecorder())
+	{
+		return Recorder->GetState();
+	}
+
+	return ETakeRecorderState::PreInitialization;
+}
+
+void UTakeRecorderSubsystem::SetTakeNumber(int32 InNewTakeNumber)
+{
+	FScopedTransaction Transaction(LOCTEXT("SetTake_Transaction", "Set Take Number"));
+	TakeMetaData->Modify();
+	TakeMetaData->SetTakeNumber(InNewTakeNumber, false);
+	bAutoApplyTakeNumber = false;
+}
+
+int32 UTakeRecorderSubsystem::GetNextTakeNumber(const FString& InSlate) const
+{
+	return UTakesCoreBlueprintLibrary::ComputeNextTakeNumber(InSlate);
+}
+
+void UTakeRecorderSubsystem::GetNumberOfTakes(const FString& InSlate, int32& OutMaxTake, int32& OutNumTakes) const
+{
+	int32 MaxTake = 0;
+
+	TArray<FAssetData> Takes = UTakesCoreBlueprintLibrary::FindTakes(InSlate);
+
+	for (const FAssetData& Asset : Takes)
+	{
+		FAssetDataTagMapSharedView::FFindTagResult TakeNumberTag = Asset.TagsAndValues.FindTag(UTakeMetaData::AssetRegistryTag_TakeNumber);
+
+		int32 ThisTakeNumber = 0;
+		if (TakeNumberTag.IsSet() && LexTryParseString(ThisTakeNumber, *TakeNumberTag.GetValue()))
+		{
+			MaxTake = FMath::Max(MaxTake, ThisTakeNumber);
+		}
+	}
+
+	OutMaxTake = MaxTake;
+	OutNumTakes = Takes.Num();
+}
+
+TArray<FAssetData> UTakeRecorderSubsystem::GetSlates(FName InPackagePath) const
+{
+	const IAssetRegistry& AssetRegistry = FModuleManager::LoadModuleChecked<FAssetRegistryModule>(AssetRegistryConstants::ModuleName).Get();
+
+	FARFilter Filter;
+	Filter.ClassPaths.Add(ULevelSequence::StaticClass()->GetClassPathName());
+	Filter.TagsAndValues.Add(UTakeMetaData::AssetRegistryTag_Slate);
+	if (!InPackagePath.IsNone())
+	{
+		Filter.PackagePaths.Add(InPackagePath);
+		Filter.bRecursivePaths = true;
+	}
+
+	TArray<FAssetData> AllAssets;
+	AssetRegistry.GetAssets(Filter, AllAssets);
+	
+	return AllAssets;
+}
+
+void UTakeRecorderSubsystem::SetSlateName(const FString& InSlateName)
+{
+	if (TakeMetaData->GetSlate() != InSlateName)
+	{
+		FScopedTransaction Transaction(LOCTEXT("SetSlate_Transaction", "Set Take Slate"));
+		TakeMetaData->Modify();
+
+		TakeMetaData->SetSlate(InSlateName);
+
+		// Compute the correct starting take number
+		const int32 NextTakeNumber = GetNextTakeNumber(TakeMetaData->GetSlate());
+		if (NextTakeNumber != TakeMetaData->GetTakeNumber())
+		{
+			TakeMetaData->SetTakeNumber(NextTakeNumber);
+		}
+	}
+}
+
+bool UTakeRecorderSubsystem::MarkFrame()
+{
+	if (UTakeRecorderBlueprintLibrary::IsRecording())
+	{
+		const FFrameRate FrameRate = TakeMetaData->GetFrameRate();
+
+		const FTimespan RecordingDuration = FDateTime::UtcNow() - TakeMetaData->GetTimestamp();
+
+		const FFrameNumber ElapsedFrame = FFrameNumber(static_cast<int32>(FrameRate.AsDecimal() * RecordingDuration.GetTotalSeconds()));
+
+		const ULevelSequence* LevelSequence = GetLevelSequence();
+		if (!LevelSequence)
+		{
+			return false;
+		}
+
+		UMovieScene* MovieScene = LevelSequence->GetMovieScene();
+		if (!MovieScene)
+		{
+			return false;
+		}
+
+		const FFrameRate DisplayRate = MovieScene->GetDisplayRate();
+		const FFrameRate TickResolution = MovieScene->GetTickResolution();
+
+		FMovieSceneMarkedFrame MarkedFrame;
+
+		const UTakeRecorderSources* Sources = LevelSequence->FindMetaData<UTakeRecorderSources>();
+		if (Sources && Sources->GetSettings().bStartAtCurrentTimecode)
+		{
+			MarkedFrame.FrameNumber = FFrameRate::TransformTime(FFrameTime(FApp::GetTimecode().ToFrameNumber(DisplayRate)), DisplayRate, TickResolution).FloorToFrame();
+		}
+		else
+		{
+			MarkedFrame.FrameNumber = ConvertFrameTime(ElapsedFrame, DisplayRate, TickResolution).CeilToFrame();
+		}
+
+		const int32 MarkedFrameIndex = MovieScene->AddMarkedFrame(MarkedFrame);
+		UTakeRecorderBlueprintLibrary::OnTakeRecorderMarkedFrameAdded(MovieScene->GetMarkedFrames()[MarkedFrameIndex]);
+		if (const UTakeRecorderSubsystem* TakeRecorderSubsystem = GEngine->GetEngineSubsystem<UTakeRecorderSubsystem>())
+		{
+			TakeRecorderSubsystem->TakeRecorderMarkedFrameAdded.Broadcast(MovieScene->GetMarkedFrames()[MarkedFrameIndex]);
+		}
+
+		return true;
+	}
+
+	return false;
+}
+
+FFrameRate UTakeRecorderSubsystem::GetFrameRate() const
+{
+	if (TakeMetaData)
+	{
+		return TakeMetaData->GetFrameRate();
+	}
+
+	return FFrameRate();
+}
+
+void UTakeRecorderSubsystem::SetFrameRate(FFrameRate InFrameRate)
+{
+	SetFrameRateImpl(InFrameRate, false);
+}
+
+void UTakeRecorderSubsystem::SetFrameRateFromTimecode()
+{
+	SetFrameRateImpl(FApp::GetTimecodeFrameRate(), true);
+}
+
+void UTakeRecorderSubsystem::ImportPreset(const FAssetData& InPreset)
+{
+	SuppliedLevelSequence = nullptr;
+	RecordIntoLevelSequence = nullptr;
+
+	UTakePreset* Take = CastChecked<UTakePreset>(InPreset.GetAsset());
+	FScopedTransaction Transaction(LOCTEXT("ImportPreset_Transaction", "Import Take Preset"));
+
+	TransientPreset->Modify();
+	TransientPreset->CopyFrom(Take);
+	TransientPreset->GetOutermost()->SetDirtyFlag(false);
+
+	GetTakeMetaData()->SetPresetOrigin(Take);
+}
+
+bool UTakeRecorderSubsystem::IsReviewing() const
+{
+	return !IsRecording() && (TakeMetaData && TakeMetaData->Recorded() && GetTakeRecorderMode() != ETakeRecorderMode::RecordIntoSequence);
+}
+
+bool UTakeRecorderSubsystem::IsRecording() const
+{
+	return UTakeRecorderBlueprintLibrary::GetActiveRecorder() ? true : false;
+}
+
+bool UTakeRecorderSubsystem::TryGetSequenceCountdown(float& OutValue) const
+{
+	const UTakeRecorder* CurrentRecording = UTakeRecorder::GetActiveRecorder();
+	const bool           bIsCountingDown  = CurrentRecording && CurrentRecording->GetState() == ETakeRecorderState::CountingDown;
+	
+	OutValue = bIsCountingDown ? CurrentRecording->GetCountdownSeconds() : 0.f;
+	return bIsCountingDown;
+}
+
+void UTakeRecorderSubsystem::SetSequenceCountdown(float InSeconds)
+{
+	if (UTakeRecorder* CurrentRecording = UTakeRecorder::GetActiveRecorder())
+	{
+		CurrentRecording->SetCountdown(InSeconds);
+	}
+}
+
+TArray<UObject*> UTakeRecorderSubsystem::GetSourceRecordSettings(UTakeRecorderSource* InSource) const
+{
+	return InSource->GetAdditionalSettingsObjects();
+}
+
+FTakeRecorderParameters UTakeRecorderSubsystem::GetGlobalRecordSettings() const
+{
+	return UTakeRecorderBlueprintLibrary::GetDefaultParameters();
+}
+
+void UTakeRecorderSubsystem::SetGlobalRecordSettings(const FTakeRecorderParameters& InParameters)
+{
+	UTakeRecorderBlueprintLibrary::SetDefaultParameters(InParameters);
+}
+
+UTakeMetaData* UTakeRecorderSubsystem::GetTakeMetaData() const
+{
+	return TakeMetaData;
+}
+
+ULevelSequence* UTakeRecorderSubsystem::GetLevelSequence() const
+{
+	if (SuppliedLevelSequence)
+	{
+		return SuppliedLevelSequence;
+	}
+	else if (RecordIntoLevelSequence)
+	{
+		return RecordIntoLevelSequence;
+	}
+	else if (RecordingLevelSequence)
+	{
+		return RecordingLevelSequence;
+	}
+	else if (TransientPreset)
+	{
+		return TransientPreset->GetLevelSequence();
+	}
+	return nullptr;
+}
+
+ULevelSequence* UTakeRecorderSubsystem::GetSuppliedLevelSequence() const
+{
+	return SuppliedLevelSequence;
+}
+
+ULevelSequence* UTakeRecorderSubsystem::GetRecordingLevelSequence() const
+{
+	return RecordingLevelSequence;
+}
+
+ULevelSequence* UTakeRecorderSubsystem::GetRecordIntoLevelSequence() const
+{
+	return RecordIntoLevelSequence;
+}
+
+ULevelSequence* UTakeRecorderSubsystem::GetLastRecordedLevelSequence() const
+{
+	return LastRecordedLevelSequence;
+}
+
+UTakePreset* UTakeRecorderSubsystem::GetTransientPreset() const
+{
+	return TransientPreset;
+}
+
+ETakeRecorderMode UTakeRecorderSubsystem::GetTakeRecorderMode() const
+{
+	if (RecordIntoLevelSequence != nullptr)
+	{
+		return ETakeRecorderMode::RecordIntoSequence;
+	}
+
+	return ETakeRecorderMode::RecordNewSequence;
+}
+
+UTakeRecorderNamingTokensData* UTakeRecorderSubsystem::GetNamingTokensData() const
+{
+	check(NamingTokensData);
+	return NamingTokensData;
+}
+
+FOnTakeRecordingInitialized& UTakeRecorderSubsystem::GetOnRecordingPreInitializedEvent()
+{
+	return OnRecordingPreInitializeEvent;
+}
+
+FOnTakeRecordingInitialized& UTakeRecorderSubsystem::GetOnRecordingInitializedEvent()
+{
+	return ORecordingInitializedEvent;
+}
+
+FOnTakeRecordingStarted& UTakeRecorderSubsystem::GetOnRecordingStartedEvent()
+{
+	return OnRecordingStartedEvent;
+}
+
+FOnTakeRecordingStopped& UTakeRecorderSubsystem::GetOnRecordingStoppedEvent()
+{
+	return OnRecordingStoppedEvent;
+}
+
+FOnTakeRecordingFinished& UTakeRecorderSubsystem::GetOnRecordingFinishedEvent()
+{
+	return OnRecordingFinishedEvent;
+}
+
+FOnTakeRecordingCancelled& UTakeRecorderSubsystem::GetOnRecordingCancelledEvent()
+{
+	return OnRecordingCancelledEvent;
+}
+
+UTakeRecorderSources::FOnSourceAdded& UTakeRecorderSubsystem::GetOnRecordingSourceAddedEvent()
+{
+	return OnRecordingSourceAddedEvent;
+}
+
+UTakeRecorderSources::FOnSourceRemoved& UTakeRecorderSubsystem::GetOnRecordingSourceRemovedEvent()
+{
+	return OnRecordingSourceRemovedEvent;
+}
+
+UTakePreset* UTakeRecorderSubsystem::AllocateTransientPreset()
+{
+	return UTakePreset::AllocateTransientPreset(GetDefault<UTakeRecorderUserSettings>()->LastOpenedPreset.Get());
+}
+
+void UTakeRecorderSubsystem::CacheMetaData()
+{
+	UTakeMetaData* NewMetaDataThisTick = nullptr;
+
+	if (const ULevelSequence* LevelSequence = GetLevelSequence())
+	{
+		NewMetaDataThisTick = LevelSequence->FindMetaData<UTakeMetaData>();
+	}
+
+	// If it's null we use the transient meta-data
+	if (!NewMetaDataThisTick)
+	{
+		// if the transient meta-data doesn't exist, create it now
+		if (!TransientTakeMetaData)
+		{
+			TransientTakeMetaData = UTakeMetaData::CreateFromDefaults(GetTransientPackage(), NAME_None);
+			TransientTakeMetaData->SetFlags(RF_Transactional | RF_Transient);
+
+			FString DefaultSlate = GetDefault<UTakeRecorderProjectSettings>()->Settings.DefaultSlate;
+			if (TransientTakeMetaData->GetSlate() != DefaultSlate)
+			{
+				TransientTakeMetaData->SetSlate(DefaultSlate, false);
+			}
+
+			// Compute the correct starting take number
+			int32 NextTakeNumber = GetNextTakeNumber(TransientTakeMetaData->GetSlate());
+			if (TransientTakeMetaData->GetTakeNumber() != NextTakeNumber)
+			{
+				TransientTakeMetaData->SetTakeNumber(NextTakeNumber, false);
+			}
+
+			UTakeMetaData::SetMostRecentMetaData(TransientTakeMetaData);
+		}
+
+		NewMetaDataThisTick = TransientTakeMetaData;
+	}
+
+	check(NewMetaDataThisTick);
+	if (NewMetaDataThisTick != TakeMetaData)
+	{
+		TakeMetaData = NewMetaDataThisTick;
+	}
+
+	//Set MovieScene Display Rate to the Preset Frame Rate.
+	if (const ULevelSequence* Sequence = GetLevelSequence())
+	{
+		if (UMovieScene* MovieScene = Sequence->GetMovieScene())
+		{
+			MovieScene->SetDisplayRate(TakeMetaData->GetFrameRate());
+		}
+	}
+
+	check(TakeMetaData);
+}
+
+void UTakeRecorderSubsystem::SetFrameRateImpl(const FFrameRate& InFrameRate, bool bFromTimecode)
+{
+	if (TakeMetaData)
+	{
+		TakeMetaData->SetFrameRateFromTimecode(bFromTimecode);
+		TakeMetaData->SetFrameRate(InFrameRate);
+	}
+	const ULevelSequence* Sequence = GetLevelSequence();
+	if (UMovieScene* MovieScene = Sequence ? Sequence->GetMovieScene() : nullptr)
+	{
+		MovieScene->SetDisplayRate(InFrameRate);
+	}
+}
+
+void UTakeRecorderSubsystem::IncrementTakeNumber()
+{
+	if (TransientTakeMetaData)
+	{
+		// Increment the transient take meta data if necessary
+		int32 NextTakeNumber = GetNextTakeNumber(TransientTakeMetaData->GetSlate());
+
+		if (TransientTakeMetaData->GetTakeNumber() != NextTakeNumber)
+		{
+			TransientTakeMetaData->SetTakeNumber(NextTakeNumber);
+		}
+	}
+	
+	// Update the preset take number at the end of recording
+	const ULevelSequence* LevelSequence = TransientPreset->GetLevelSequence();
+	if (UTakeMetaData* MetaData = LevelSequence ? LevelSequence->FindMetaData<UTakeMetaData>() : nullptr)
+	{
+		const int32 NextTakeNumber = GetNextTakeNumber(MetaData->GetSlate());
+		MetaData->SetTakeNumber(NextTakeNumber);
+	}
+
+	bAutoApplyTakeNumber = true;
+}
+
+void UTakeRecorderSubsystem::OnAssetRegistryFilesLoaded()
+{
+	if (bAutoApplyTakeNumber && TransientTakeMetaData)
+	{
+		const int32 NextTakeNumber = GetNextTakeNumber(TransientTakeMetaData->GetSlate());
+		TransientTakeMetaData->SetTakeNumber(NextTakeNumber);
+	}
+}
+
+void UTakeRecorderSubsystem::OnRecordingInitialized(UTakeRecorder* Recorder)
+{
+	// This needs to be stored with a strong ptr before our panels refresh, otherwise the weak sequencer ref will be lost during initialize.
+	RecordingLevelSequence = Recorder->GetSequence();
+	// Recache the meta-data here since we know that the sequence has probably changed as a result of the recording being started
+	CacheMetaData();
+	OnRecordingFinishedHandle = Recorder->OnRecordingFinished().AddUObject(this, &UTakeRecorderSubsystem::OnRecordingFinished);
+	OnRecordingCancelledHandle = Recorder->OnRecordingCancelled().AddUObject(this, &UTakeRecorderSubsystem::OnRecordingCancelled);
+
+	GetOnRecordingInitializedEvent().Broadcast(Recorder);
+	TakeRecorderInitialized.Broadcast();
+}
+
+void UTakeRecorderSubsystem::OnRecordingStarted(UTakeRecorder* Recorder)
+{
+	GetOnRecordingStartedEvent().Broadcast(Recorder);
+	TakeRecorderStarted.Broadcast();
+}
+
+void UTakeRecorderSubsystem::OnRecordingStopped(UTakeRecorder* Recorder)
+{
+	Recorder->OnRecordingStopped().Remove(OnRecordingStoppedHandle);
+	GetOnRecordingStoppedEvent().Broadcast(Recorder);
+	TakeRecorderStopped.Broadcast();
+}
+
+void UTakeRecorderSubsystem::OnRecordingFinished(UTakeRecorder* Recorder)
+{
+	LastRecordedLevelSequence = RecordingLevelSequence;
+	RecordingLevelSequence = nullptr;
+
+	IncrementTakeNumber();
+	
+	if (TransientTakeMetaData)
+	{
+		// Increment the transient take meta data if necessary
+		int32 NextTakeNumber = GetNextTakeNumber(TransientTakeMetaData->GetSlate());
+
+		if (TransientTakeMetaData->GetTakeNumber() != NextTakeNumber)
+		{
+			TransientTakeMetaData->SetTakeNumber(NextTakeNumber);
+		}
+
+		bAutoApplyTakeNumber = true;
+	}
+
+	Recorder->OnRecordingFinished().Remove(OnRecordingFinishedHandle);
+
+	GetOnRecordingFinishedEvent().Broadcast(Recorder);
+	TakeRecorderFinished.Broadcast(LastRecordedLevelSequence);
+}
+
+void UTakeRecorderSubsystem::OnRecordingCancelled(UTakeRecorder* Recorder)
+{
+	RecordingLevelSequence = nullptr;
+
+	Recorder->OnRecordingFinished().Remove(OnRecordingFinishedHandle);
+	Recorder->OnRecordingCancelled().Remove(OnRecordingCancelledHandle);
+
+	GetOnRecordingCancelledEvent().Broadcast(Recorder);
+	TakeRecorderCancelled.Broadcast();
+}
+
+void UTakeRecorderSubsystem::SetLastLevelSequence(ULevelSequence* InSequence)
+{
+	LastRecordedLevelSequence = InSequence;
+}
+
+void UTakeRecorderSubsystem::OnSourceAdded(UTakeRecorderSource* Source)
+{
+	GetOnRecordingSourceAddedEvent().Broadcast(Source);
+	TakeRecorderSourceAdded.Broadcast(Source);
+}
+
+void UTakeRecorderSubsystem::OnSourceRemoved(UTakeRecorderSource* Source)
+{
+	GetOnRecordingSourceRemovedEvent().Broadcast(Source);
+	TakeRecorderSourceRemoved.Broadcast(Source);
+}
+
+void UTakeRecorderSubsystem::BindToNamingTokenEvents()
+{
+	if (GEngine && !GetNamingTokensData()->TakeRecorderNamingTokens.IsValid())
+	{
+		GetNamingTokensData()->TakeRecorderNamingTokens = GEngine->GetEngineSubsystem<UNamingTokensEngineSubsystem>()->GetNamingTokens(ITakeRecorderNamingTokensModule::GetTakeRecorderNamespace());
+		if (ensure(GetNamingTokensData()->TakeRecorderNamingTokens.IsValid()))
+		{
+			const TStrongObjectPtr<UNamingTokens> TakeRecorderNamingTokensPin = GetNamingTokensData()->TakeRecorderNamingTokens.Pin();
+			TakeRecorderNamingTokensPin->RegisterExternalTokens(GetNamingTokensData()->NamingTokensExternalGuid);
+			TakeRecorderNamingTokensPin->GetOnPreEvaluateEvent().AddUObject(this, &UTakeRecorderSubsystem::OnTakeRecorderNamingTokensPreEvaluate);
+		}
+	}
+}
+
+void UTakeRecorderSubsystem::UnbindNamingTokensEvents()
+{
+	if (const TStrongObjectPtr<UNamingTokens> TakeRecorderNamingTokensPin = GetNamingTokensData()->TakeRecorderNamingTokens.Pin())
+	{
+		TakeRecorderNamingTokensPin->UnregisterExternalTokens(GetNamingTokensData()->NamingTokensExternalGuid);
+		TakeRecorderNamingTokensPin->GetOnPreEvaluateEvent().RemoveAll(this);
+	}
+}
+
+void UTakeRecorderSubsystem::OnTakeRecorderNamingTokensPreEvaluate(const FNamingTokensEvaluationData& InEvaluationData)
+{
+	if (GetNamingTokensData()->TakeRecorderNamingTokens.IsValid())
+	{
+		TArray<FNamingTokenData>& ExternalTokens = GetNamingTokensData()->TakeRecorderNamingTokens->GetExternalTokensChecked(GetNamingTokensData()->NamingTokensExternalGuid);
+
+		ExternalTokens.Reset(GetNamingTokensData()->UserDefinedTokens.Num());
+		for (const TTuple<FNamingTokenData, FText>& UserToken : GetNamingTokensData()->UserDefinedTokens)
+		{
+			ExternalTokens.Add({
+				UserToken.Key.TokenKey,
+				UserToken.Key.DisplayName,
+				// @todo naming tokens - If we support TOptional<FText> that can be used instead of a lambda
+				FNamingTokenData::FTokenProcessorDelegateNative::CreateLambda([Value = UserToken.Value]() {
+					return Value;
+				})
+			});
+		}
+	}
+}
 
 #undef LOCTEXT_NAMESPACE
diff --git a/Engine/Plugins/VirtualProduction/Takes/Source/TakeRecorder/Private/TakeRecorderModule.cpp b/Engine/Plugins/VirtualProduction/Takes/Source/TakeRecorder/Private/TakeRecorderModule.cpp
--- a/Engine/Plugins/VirtualProduction/Takes/Source/TakeRecorder/Private/TakeRecorderModule.cpp
+++ b/Engine/Plugins/VirtualProduction/Takes/Source/TakeRecorder/Private/TakeRecorderModule.cpp
@@ -514,6 +514,16 @@ FDelegateHandle FTakeRecorderModule::RegisterSourcesMenuExtension(const FOnExten
 	return SourcesMenuExtenderEvent.Add(InExtension);
 }
 
+void FTakeRecorderModule::RegisterSourcesExtension(const FSourceExtensionData& InData)
+{
+	SourceExtensionData = InData;
+}
+
+void FTakeRecorderModule::UnregisterSourcesExtension()
+{
+	SourceExtensionData = FSourceExtensionData();
+}
+
 void FTakeRecorderModule::UnregisterSourcesMenuExtension(FDelegateHandle Handle)
 {
 	SourcesMenuExtenderEvent.Remove(Handle);
diff --git a/Engine/Plugins/VirtualProduction/Takes/Source/TakeRecorder/Private/TakeRecorderModule.h b/Engine/Plugins/VirtualProduction/Takes/Source/TakeRecorder/Private/TakeRecorderModule.h
--- a/Engine/Plugins/VirtualProduction/Takes/Source/TakeRecorder/Private/TakeRecorderModule.h
+++ b/Engine/Plugins/VirtualProduction/Takes/Source/TakeRecorder/Private/TakeRecorderModule.h
@@ -40,12 +40,20 @@ public:
 	{
 		return ForceSaveAsPresetEvent;
 	}
+
+	FSourceExtensionData& GetSourcesExtensionData()
+	{
+		return SourceExtensionData;
+	}
+	
 private:
 
 	void StartupModule() override;
 	void ShutdownModule() override;
 
 	FDelegateHandle RegisterSourcesMenuExtension(const FOnExtendSourcesMenu& InExtension) override;
+	virtual void RegisterSourcesExtension(const FSourceExtensionData& InData) override;
+	virtual void UnregisterSourcesExtension() override;
 	void UnregisterSourcesMenuExtension(FDelegateHandle Handle) override;
 	void RegisterSettingsObject(UObject* InSettingsObject) override;
 
@@ -109,4 +117,10 @@ private:
 	TArray<TWeakObjectPtr<>> ExternalObjects;
 
 	TObjectPtr<USequencerSettings> SequencerSettings;
+
+	/**
+	 * Registered extensions for the source module. We only use a single struct rather than container
+	 * as only the source module should be utilizing this.
+	 */
+	FSourceExtensionData SourceExtensionData;
 };
diff --git a/Engine/Plugins/VirtualProduction/Takes/Source/TakeRecorder/Private/Widgets/STakeRecorderCockpit.cpp b/Engine/Plugins/VirtualProduction/Takes/Source/TakeRecorder/Private/Widgets/STakeRecorderCockpit.cpp
--- a/Engine/Plugins/VirtualProduction/Takes/Source/TakeRecorder/Private/Widgets/STakeRecorderCockpit.cpp
+++ b/Engine/Plugins/VirtualProduction/Takes/Source/TakeRecorder/Private/Widgets/STakeRecorderCockpit.cpp
@@ -76,6 +76,7 @@
 #include "NamingTokens.h"
 #include "NamingTokensEngineSubsystem.h"
 #include "NamingTokensSpecifiers.h"
+#include "TakeRecorderNamingTokensData.h"
 
 extern UNREALED_API UEditorEngine* GEditor;
 
@@ -83,30 +84,24 @@ extern UNREALED_API UEditorEngine* GEditor;
 
 STakeRecorderCockpit::~STakeRecorderCockpit()
 {
-	UTakeRecorder::OnRecordingInitialized().Remove(OnRecordingInitializedHandle);
 	FEditorDelegates::PostUndoRedo.Remove((PostUndoRedoDelegateHandle));
 
-	if (FAssetRegistryModule* AssetRegistryModule = FModuleManager::GetModulePtr<FAssetRegistryModule>("AssetRegistry"))
-	{
-		IAssetRegistry* AssetRegistry = AssetRegistryModule->TryGet();
-		if (AssetRegistry)
-		{
-			AssetRegistry->OnFilesLoaded().Remove(OnAssetRegistryFilesLoadedHandle);
-		}
-	}
-
 	if (!ensure(TransactionIndex == INDEX_NONE))
 	{
 		GEditor->CancelTransaction(TransactionIndex);
 	}
+}
 
-	UnbindNamingTokensEvents();
+UTakeMetaData* STakeRecorderCockpit::GetMetaData() const
+{
+	return TakeRecorderSubsystem->GetTakeMetaData();
 }
 
-void STakeRecorderCockpit::AddReferencedObjects(FReferenceCollector& Collector)
+UTakeMetaData* STakeRecorderCockpit::GetMetaDataChecked() const
 {
-	Collector.AddReferencedObject(TakeMetaData);
-	Collector.AddReferencedObject(TransientTakeMetaData);
+	UTakeMetaData* TakeMetaData = GetMetaData();
+	check(TakeMetaData);
+	return TakeMetaData;
 }
 
 struct SNonThrottledButton : SButton
@@ -160,45 +155,19 @@ namespace UE::TakeRecorder::Private
 
 void STakeRecorderCockpit::Construct(const FArguments& InArgs)
 {
-	FAssetRegistryModule& AssetRegistryModule = FModuleManager::LoadModuleChecked<FAssetRegistryModule>(TEXT("AssetRegistry"));
-	OnAssetRegistryFilesLoadedHandle = AssetRegistryModule.Get().OnFilesLoaded().AddSP(this, &STakeRecorderCockpit::OnAssetRegistryFilesLoaded);
-	
-	OnRecordingInitializedHandle = UTakeRecorder::OnRecordingInitialized().AddSP(this, &STakeRecorderCockpit::OnRecordingInitialized);
-
-	bAutoApplyTakeNumber = true;
-
-	TakeMetaData = nullptr;
-	TransientTakeMetaData = nullptr;
-	NamingTokensData = static_cast<TStrongObjectPtr<UTakeRecorderNamingTokensWidgetData>>(NewObject<UTakeRecorderNamingTokensWidgetData>());
-	NamingTokensData->SetFlags(RF_Transactional);
-
 	PostUndoRedoDelegateHandle = FEditorDelegates::PostUndoRedo.AddRaw(this, &STakeRecorderCockpit::PostUndoRedo);
 
 	CachedTakeSlate.Empty();
 	CachedTakeNumber = -1;
 
-	LevelSequenceAttribute = InArgs._LevelSequence;
-
-	TakeRecorderModeAttribute = InArgs._TakeRecorderMode;
-
-	CacheMetaData();
-
-	if (TakeMetaData && !TakeMetaData->IsLocked())
-	{
-		int32 NextTakeNumber = UTakesCoreBlueprintLibrary::ComputeNextTakeNumber(TakeMetaData->GetSlate());
-		if (NextTakeNumber != TakeMetaData->GetTakeNumber())
-		{
-			TakeMetaData->SetTakeNumber(NextTakeNumber);
-		}
-	}
-
+	TakeRecorderSubsystem = GEngine->GetEngineSubsystem<UTakeRecorderSubsystem>();
+	
 	UpdateTakeError();
 	UpdateRecordError();
 
 	DigitsTypeInterface = MakeShareable(new FDigitsTypeInterface);
 
 	BindCommands();
-	BindToNamingTokenEvents();
 
 	TransactionIndex = INDEX_NONE;
 
@@ -537,14 +506,14 @@ void STakeRecorderCockpit::UpdateRecordError()
 		return;
 	}
 
-	ULevelSequence* Sequence = LevelSequenceAttribute.Get();
+	ULevelSequence* Sequence = TakeRecorderSubsystem->GetLevelSequence();
 	if (!Sequence)
 	{
 		RecordErrorText = LOCTEXT("ErrorWidget_NoSequence", "There is no sequence to record from. Please re-open Take Recorder.");
 		return;
 	}
 
-	if (!Sequence->HasAnyFlags(RF_Transient) && TakeRecorderModeAttribute.Get() != ETakeRecorderMode::RecordIntoSequence)
+	if (!Sequence->HasAnyFlags(RF_Transient) && TakeRecorderSubsystem->GetTakeRecorderMode() != ETakeRecorderMode::RecordIntoSequence)
 	{
 		RecordErrorText = FText();
 		return;
@@ -562,7 +531,9 @@ void STakeRecorderCockpit::UpdateRecordError()
 		return;
 	}
 
-	if (TakeMetaData->GetSlate().IsEmpty())
+	UTakeMetaData* TakeMetaData = GetMetaData();
+
+	if (TakeMetaData && TakeMetaData->GetSlate().IsEmpty())
 	{
 		RecordErrorText = LOCTEXT("ErrorWidget_NoSlate", "You must enter a slate to begin recording.");
 		return;
@@ -598,6 +569,8 @@ void STakeRecorderCockpit::UpdateTakeError()
 {
 	TakeErrorText = FText();
 
+	const UTakeMetaData* TakeMetaData = GetMetaDataChecked();
+	
 	TArray<FAssetData> DuplicateTakes = UTakesCoreBlueprintLibrary::FindTakes(TakeMetaData->GetSlate(), TakeMetaData->GetTakeNumber());
 
 	// If there's only a single one, and it's the one that we're looking at directly, don't show the error
@@ -663,70 +636,6 @@ TSharedRef<SWidget> STakeRecorderCockpit::OnRecordingOptionsMenu()
 	return MenuBuilder.MakeWidget();
 }
 
-void STakeRecorderCockpit::CacheMetaData()
-{
-	UTakeMetaData* NewMetaDataThisTick = nullptr;
-
-	if (ULevelSequence* LevelSequence = LevelSequenceAttribute.Get())
-	{
-		NewMetaDataThisTick = LevelSequence->FindMetaData<UTakeMetaData>();
-	}
-
-	// If it's null we use the transient meta-data
-	if (!NewMetaDataThisTick)
-	{
-		// if the transient meta-data doesn't exist, create it now
-		if (!TransientTakeMetaData)
-		{
-			TransientTakeMetaData = UTakeMetaData::CreateFromDefaults(GetTransientPackage(), NAME_None);
-			TransientTakeMetaData->SetFlags(RF_Transactional | RF_Transient);
-
-			FString DefaultSlate = GetDefault<UTakeRecorderProjectSettings>()->Settings.DefaultSlate;
-			if (TransientTakeMetaData->GetSlate() != DefaultSlate)
-			{
-				TransientTakeMetaData->SetSlate(DefaultSlate, false);
-			}
-
-			// Compute the correct starting take number
-			int32 NextTakeNumber = UTakesCoreBlueprintLibrary::ComputeNextTakeNumber(TransientTakeMetaData->GetSlate());
-			if (TransientTakeMetaData->GetTakeNumber() != NextTakeNumber)
-			{
-				TransientTakeMetaData->SetTakeNumber(NextTakeNumber, false);
-			}
-
-			UTakeMetaData::SetMostRecentMetaData(TransientTakeMetaData);
-		}
-
-		NewMetaDataThisTick = TransientTakeMetaData;
-	}
-
-	check(NewMetaDataThisTick);
-	if (NewMetaDataThisTick != TakeMetaData)
-	{
-		TakeMetaData = NewMetaDataThisTick;
-	}
-
-	if (TakeMetaData->GetSlate() != CachedTakeSlate || TakeMetaData->GetTakeNumber() != CachedTakeNumber)
-	{
-		CachedTakeNumber = TakeMetaData->GetTakeNumber();
-		CachedTakeSlate = TakeMetaData->GetSlate();
-
-		// Previously, the take error would be updated in Tick(), but the asset registry can be slow, 
-		// so it should be sufficient to update it only when the slate changes.
-		UpdateTakeError();
-	}
-
-	//Set MovieScene Display Rate to the Preset Frame Rate.
-	ULevelSequence* Sequence = LevelSequenceAttribute.Get();
-	UMovieScene*    MovieScene = Sequence ? Sequence->GetMovieScene() : nullptr;
-	if (MovieScene)
-	{
-		MovieScene->SetDisplayRate(TakeMetaData->GetFrameRate());
-	}
-
-	check(TakeMetaData);
-}
-
 void STakeRecorderCockpit::Tick(const FGeometry& AllottedGeometry, const double InCurrentTime, const float InDeltaTime)
 {
 	Refresh();
@@ -744,7 +653,19 @@ void STakeRecorderCockpit::PostUndoRedo()
 
 void STakeRecorderCockpit::Refresh()
 {
-	CacheMetaData();
+	if (const UTakeMetaData* TakeMetaData = GetMetaData())
+	{
+		if (TakeMetaData->GetSlate() != CachedTakeSlate || TakeMetaData->GetTakeNumber() != CachedTakeNumber)
+		{
+			CachedTakeNumber = TakeMetaData->GetTakeNumber();
+			CachedTakeSlate = TakeMetaData->GetSlate();
+
+			// Previously, the take error would be updated in Tick(), but the asset registry can be slow, 
+			// so it should be sufficient to update it only when the slate changes.
+			UpdateTakeError();
+		}
+	}
+	
 	UpdateRecordError();
 }
 
@@ -806,7 +727,7 @@ void STakeRecorderCockpit::NotifyDetailsViewAdded(const TWeakPtr<IDetailsView>&
 
 FText STakeRecorderCockpit::GetSlateText() const
 {
-	return FText::FromString(TakeMetaData->GetSlate());
+	return FText::FromString(GetMetaDataChecked()->GetSlate());
 }
 
 FText STakeRecorderCockpit::GetTimecodeText() const
@@ -816,12 +737,14 @@ FText STakeRecorderCockpit::GetTimecodeText() const
 
 FText STakeRecorderCockpit::GetUserDescriptionText() const
 {
+	const UTakeMetaData* TakeMetaData = GetMetaDataChecked();
 	return FText::FromString(TakeMetaData->GetDescription());
 }
 
 FText STakeRecorderCockpit::GetTimestampText() const
 {
 	// If not recorded, return current time
+	const UTakeMetaData* TakeMetaData = GetMetaDataChecked();
 	if (TakeMetaData->GetTimestamp() == FDateTime(0))
 	{
 		return FText::AsDateTime(FDateTime::UtcNow());
@@ -835,6 +758,7 @@ FText STakeRecorderCockpit::GetTimestampText() const
 FText STakeRecorderCockpit::GetTimestampTooltipText() const
 {
 	// If not recorded, return current time
+	const UTakeMetaData* TakeMetaData = GetMetaDataChecked();
 	if (TakeMetaData->GetTimestamp() == FDateTime(0))
 	{
 		return LOCTEXT("CurrentTimestamp", "The current date/time");
@@ -847,21 +771,18 @@ FText STakeRecorderCockpit::GetTimestampTooltipText() const
 
 FText STakeRecorderCockpit::GetEvaluatedTakeSaveDirText() const
 {
-	return NamingTokensData->EvaluatedTextValue;
+	return TakeRecorderSubsystem->GetNamingTokensData()->EvaluatedTextValue;
 }
 
 void STakeRecorderCockpit::SetFrameRate(FFrameRate InFrameRate, bool bFromTimecode)
 {
-	if (TakeMetaData)
+	if (bFromTimecode)
 	{
-		TakeMetaData->SetFrameRateFromTimecode(bFromTimecode);
-		TakeMetaData->SetFrameRate(InFrameRate);
+		TakeRecorderSubsystem->SetFrameRateFromTimecode();
 	}
-	ULevelSequence* Sequence = LevelSequenceAttribute.Get();
-	UMovieScene*    MovieScene = Sequence ? Sequence->GetMovieScene() : nullptr;
-	if (MovieScene)
+	else
 	{
-		MovieScene->SetDisplayRate(InFrameRate);
+		TakeRecorderSubsystem->SetFrameRate(InFrameRate);
 	}
 }
 
@@ -872,7 +793,7 @@ bool STakeRecorderCockpit::IsSameFrameRate(FFrameRate InFrameRate) const
 
 FFrameRate STakeRecorderCockpit::GetFrameRate() const
 {
-	return TakeMetaData->GetFrameRate();
+	return TakeRecorderSubsystem->GetFrameRate();
 }
 
 FText STakeRecorderCockpit::GetFrameRateText() const
@@ -887,7 +808,7 @@ FText STakeRecorderCockpit::GetFrameRateTooltipText() const
 
 bool STakeRecorderCockpit::IsFrameRateCompatible(FFrameRate InFrameRate) const
 {
-	ULevelSequence* Sequence   = LevelSequenceAttribute.Get();
+	ULevelSequence* Sequence   = TakeRecorderSubsystem->GetLevelSequence();
 	UMovieScene*    MovieScene = Sequence ? Sequence->GetMovieScene() : nullptr;
 
 	return MovieScene && InFrameRate.IsMultipleOf(MovieScene->GetTickResolution());
@@ -895,32 +816,22 @@ bool STakeRecorderCockpit::IsFrameRateCompatible(FFrameRate InFrameRate) const
 
 bool STakeRecorderCockpit::IsSetFromTimecode() const
 {
-	return TakeMetaData->GetFrameRateFromTimecode();
+	return GetMetaDataChecked()->GetFrameRateFromTimecode();
 }
 
 void STakeRecorderCockpit::SetSlateText(const FText& InNewText, ETextCommit::Type InCommitType)
 {
+	UTakeMetaData* TakeMetaData = GetMetaDataChecked();
 	if (TakeMetaData->GetSlate() != InNewText.ToString())
 	{
-		FScopedTransaction Transaction(LOCTEXT("SetSlate_Transaction", "Set Take Slate"));
-		TakeMetaData->Modify();
-
-		TakeMetaData->SetSlate(InNewText.ToString());
-
-		// Compute the correct starting take number
-		int32 NextTakeNumber = UTakesCoreBlueprintLibrary::ComputeNextTakeNumber(TakeMetaData->GetSlate());
-
-		if (NextTakeNumber != TakeMetaData->GetTakeNumber())
-		{
-			TakeMetaData->SetTakeNumber(NextTakeNumber);
-		}
-
+		TakeRecorderSubsystem->SetSlateName(InNewText.ToString());
 		OnTokenValueUpdated();
 	}
 }
 
 void STakeRecorderCockpit::SetUserDescriptionText(const FText& InNewText, ETextCommit::Type InCommitType)
 {
+	UTakeMetaData* TakeMetaData = GetMetaDataChecked();
 	if (TakeMetaData->GetDescription() != InNewText.ToString())
 	{
 		FScopedTransaction Transaction(LOCTEXT("SetDescription_Transaction", "Set Description"));
@@ -932,13 +843,13 @@ void STakeRecorderCockpit::SetUserDescriptionText(const FText& InNewText, ETextC
 
 int32 STakeRecorderCockpit::GetTakeNumber() const
 {
-	return TakeMetaData->GetTakeNumber();
+	return GetMetaDataChecked()->GetTakeNumber();
 }
 
 FReply STakeRecorderCockpit::OnSetNextTakeNumber()
 {
-	int32 NextTakeNumber = UTakesCoreBlueprintLibrary::ComputeNextTakeNumber(TakeMetaData->GetSlate());
-
+	UTakeMetaData* TakeMetaData = GetMetaDataChecked();
+	const int32 NextTakeNumber = UTakesCoreBlueprintLibrary::ComputeNextTakeNumber(TakeMetaData->GetSlate());
 	if (TakeMetaData->GetTakeNumber() != NextTakeNumber)
 	{
 		FScopedTransaction Transaction(LOCTEXT("SetNextTakeNumber_Transaction", "Set Next Take Number"));
@@ -960,6 +871,7 @@ void STakeRecorderCockpit::OnBeginSetTakeNumber()
 	}
 
 	TransactionIndex = GEditor->BeginTransaction(nullptr, LOCTEXT("SetTakeNumber_Transaction", "Set Take Number"), nullptr);
+	UTakeMetaData* TakeMetaData = GetMetaDataChecked();
 	TakeMetaData->Modify();
 }
 
@@ -969,8 +881,7 @@ void STakeRecorderCockpit::SetTakeNumber(int32 InNewTakeNumber)
 
 	if (TransactionIndex != INDEX_NONE || bIsInPIEOrSimulate)
 	{
-		TakeMetaData->SetTakeNumber(InNewTakeNumber, false);
-		bAutoApplyTakeNumber = false;
+		TakeRecorderSubsystem->SetTakeNumber(InNewTakeNumber);
 	}
 	
 	OnTokenValueUpdated();
@@ -980,6 +891,7 @@ void STakeRecorderCockpit::SetTakeNumber_FromCommit(int32 InNewTakeNumber, EText
 {
 	const bool bIsInPIEOrSimulate = GEditor->PlayWorld != NULL || GEditor->bIsSimulatingInEditor;
 
+	UTakeMetaData* TakeMetaData = GetMetaDataChecked();
 	if (TransactionIndex == INDEX_NONE && !bIsInPIEOrSimulate)
 	{
 		if (TakeMetaData->GetTakeNumber() != InNewTakeNumber)
@@ -993,7 +905,6 @@ void STakeRecorderCockpit::SetTakeNumber_FromCommit(int32 InNewTakeNumber, EText
 		TakeMetaData->SetTakeNumber(InNewTakeNumber);
 	}
 
-	bAutoApplyTakeNumber = false;
 	OnTokenValueUpdated();
 }
 
@@ -1005,7 +916,7 @@ void STakeRecorderCockpit::OnEndSetTakeNumber(int32 InFinalValue)
 	{
 		check(TransactionIndex != INDEX_NONE);
 	}
-
+	UTakeMetaData* TakeMetaData = GetMetaDataChecked();
 	TakeMetaData->SetTakeNumber(InFinalValue);
 
 	GEditor->EndTransaction();
@@ -1025,60 +936,18 @@ void STakeRecorderCockpit::SetEngineTimeDilation(float InEngineTimeDilation)
 
 FReply STakeRecorderCockpit::OnAddMarkedFrame()
 {
-	if (UTakeRecorderBlueprintLibrary::IsRecording())
-	{
-		FFrameRate FrameRate = TakeMetaData->GetFrameRate();
-
-		FTimespan RecordingDuration = FDateTime::UtcNow() - TakeMetaData->GetTimestamp();
-
-		FFrameNumber ElapsedFrame = FFrameNumber(static_cast<int32>(FrameRate.AsDecimal() * RecordingDuration.GetTotalSeconds()));
-		
-		ULevelSequence* LevelSequence = LevelSequenceAttribute.Get();
-		if (!LevelSequence)
-		{
-			return FReply::Handled();
-		}
-
-		UMovieScene* MovieScene = LevelSequence->GetMovieScene();
-		if (!MovieScene)
-		{
-			return FReply::Handled();
-		}
-
-		FFrameRate DisplayRate = MovieScene->GetDisplayRate();
-		FFrameRate TickResolution = MovieScene->GetTickResolution();
-
-		FMovieSceneMarkedFrame MarkedFrame;
-
-		UTakeRecorderSources* Sources = LevelSequence->FindMetaData<UTakeRecorderSources>();
-		if (Sources && Sources->GetSettings().bStartAtCurrentTimecode)
-		{
-			MarkedFrame.FrameNumber = FFrameRate::TransformTime(FFrameTime(FApp::GetTimecode().ToFrameNumber(DisplayRate)), DisplayRate, TickResolution).FloorToFrame();
-		}
-		else
-		{
-			MarkedFrame.FrameNumber = ConvertFrameTime(ElapsedFrame, DisplayRate, TickResolution).CeilToFrame();
-		}
-		
-		int32 MarkedFrameIndex = MovieScene->AddMarkedFrame(MarkedFrame);
-		UTakeRecorderBlueprintLibrary::OnTakeRecorderMarkedFrameAdded(MovieScene->GetMarkedFrames()[MarkedFrameIndex]);
-		if (UTakeRecorderSubsystem* TakeRecorderSubsystem = GEngine->GetEngineSubsystem<UTakeRecorderSubsystem>())
-		{
-			TakeRecorderSubsystem->TakeRecorderMarkedFrameAdded.Broadcast(MovieScene->GetMarkedFrames()[MarkedFrameIndex]);
-		}
-	}
-
+	TakeRecorderSubsystem->MarkFrame();
 	return FReply::Handled();
 }
 
 bool STakeRecorderCockpit::Reviewing() const 
 {
-	return bool(!Recording() && (TakeMetaData && TakeMetaData->Recorded() && TakeRecorderModeAttribute.Get() != ETakeRecorderMode::RecordIntoSequence));
+	return TakeRecorderSubsystem->IsReviewing();
 }
 
 bool STakeRecorderCockpit::Recording() const
 {
-	return UTakeRecorderBlueprintLibrary::GetActiveRecorder() ? true : false;
+	return TakeRecorderSubsystem->IsRecording();
 }
 
 ECheckBoxState STakeRecorderCockpit::IsRecording() const
@@ -1093,7 +962,7 @@ bool STakeRecorderCockpit::CanRecord() const
 
 bool STakeRecorderCockpit::IsLocked() const 
 {
-	if (TakeMetaData != nullptr)
+	if (const UTakeMetaData* TakeMetaData = GetMetaData())
 	{
 		return TakeMetaData->IsLocked();
 	}
@@ -1102,7 +971,7 @@ bool STakeRecorderCockpit::IsLocked() const
 
 void STakeRecorderCockpit::OnToggleRecording(ECheckBoxState)
 {
-	ULevelSequence*       LevelSequence = LevelSequenceAttribute.Get();
+	ULevelSequence*       LevelSequence = TakeRecorderSubsystem->GetLevelSequence();
 	UTakeRecorderSources* Sources = LevelSequence ? LevelSequence->FindMetaData<UTakeRecorderSources>() : nullptr;
 
 	UTakeRecorder* CurrentRecording = UTakeRecorder::GetActiveRecorder();
@@ -1118,91 +987,22 @@ void STakeRecorderCockpit::OnToggleRecording(ECheckBoxState)
 
 void STakeRecorderCockpit::StopRecording()
 {
-	UTakeRecorder* CurrentRecording = UTakeRecorder::GetActiveRecorder();
-	if (CurrentRecording)
-	{
-		CurrentRecording->Stop();
-	}
+	TakeRecorderSubsystem->StopRecording();
 }
 
 void STakeRecorderCockpit::CancelRecording()
 {
-	UTakeRecorder* CurrentRecording = UTakeRecorder::GetActiveRecorder();
-	if (CurrentRecording)
-	{
-		TSharedRef<SMessageDialog> ConfirmDialog = SNew(SMessageDialog)
-			.Title(FText(LOCTEXT("ConfirmCancelRecordingTitle", "Cancel Recording?")))
-			.Message(LOCTEXT("ConfirmCancelRecording", "Are you sure you want to cancel the current recording?"))
-			.Buttons({
-				SCustomDialog::FButton(LOCTEXT("Yes", "Yes"), FSimpleDelegate::CreateLambda([WeakRecording = TWeakObjectPtr<UTakeRecorder>(CurrentRecording)]()
-					{ 
-						if (WeakRecording.IsValid())
-						{
-							WeakRecording.Get()->Cancel(); 
-						}
-					})),
-				SCustomDialog::FButton(LOCTEXT("No", "No"))
-			});
-
-		// Non modal so that the recording continues to update
-		ConfirmDialog->Show();
-	}
+	TakeRecorderSubsystem->CancelRecording();
 }
 
 void STakeRecorderCockpit::StartRecording()
 {
-	static bool bStartedRecording = false;
-
-	if (bStartedRecording)
-	{
-		return;
-	}
-
-	TGuardValue<bool> ReentrantGuard(bStartedRecording, true);
-
-	ULevelSequence*       LevelSequence = LevelSequenceAttribute.Get();
-	UTakeRecorderSources* Sources = LevelSequence ? LevelSequence->FindMetaData<UTakeRecorderSources>() : nullptr;
-	
-	if (LevelSequence && Sources)
-	{
-		FTakeRecorderParameters Parameters;
-		Parameters.User    = GetDefault<UTakeRecorderUserSettings>()->Settings;
-		Parameters.Project = GetDefault<UTakeRecorderProjectSettings>()->Settings;
-		Parameters.TakeRecorderMode = TakeRecorderModeAttribute.Get();
-		Parameters.StartFrame = LevelSequence->GetMovieScene()->GetPlaybackRange().GetLowerBoundValue();
-
-		FText ErrorText = LOCTEXT("UnknownError", "An unknown error occurred when trying to start recording");
-
-		IAssetEditorInstance* AssetEditor = GEditor->GetEditorSubsystem<UAssetEditorSubsystem>()->FindEditorForAsset(LevelSequence, false);
-		ILevelSequenceEditorToolkit* LevelSequenceEditor = static_cast<ILevelSequenceEditorToolkit*>(AssetEditor);
-
-		if (LevelSequenceEditor && LevelSequenceEditor->GetSequencer())
-		{
-			// If not resetting the playhead, store the current time as the start frame for recording. 
-			// This will ultimately be the start of the playback range and the recording will begin from that time.
-			if (!Parameters.User.bResetPlayhead)
-			{
-				Parameters.StartFrame = LevelSequenceEditor->GetSequencer()->GetLocalTime().Time.FrameNumber;
-			}
-		}
-
-		UTakeRecorder* NewRecorder = NewObject<UTakeRecorder>(GetTransientPackage(), NAME_None, RF_Transient);
-
-		if (!NewRecorder->Initialize(LevelSequence, Sources, TakeMetaData, Parameters, &ErrorText))
-		{
-			if (ensure(!ErrorText.IsEmpty()))
-			{
-				FNotificationInfo Info(ErrorText);
-				Info.ExpireDuration = 5.0f;
-				FSlateNotificationManager::Get().AddNotification(Info)->SetCompletionState(SNotificationItem::CS_Fail);
-			}
-		}
-	}
+	TakeRecorderSubsystem->StartRecording();
 }
 
 FReply STakeRecorderCockpit::NewRecordingFromThis()
 {
-	ULevelSequence* Sequence = LevelSequenceAttribute.Get();
+	ULevelSequence* Sequence = TakeRecorderSubsystem->GetLevelSequence();
 	if (!Sequence)
 	{
 		return FReply::Unhandled();
@@ -1219,41 +1019,6 @@ FReply STakeRecorderCockpit::NewRecordingFromThis()
 	return FReply::Handled();
 }
 
-void STakeRecorderCockpit::OnAssetRegistryFilesLoaded()
-{
-	if (bAutoApplyTakeNumber && TransientTakeMetaData)
-	{
-		int32 NextTakeNumber = UTakesCoreBlueprintLibrary::ComputeNextTakeNumber(TransientTakeMetaData->GetSlate());
-		TransientTakeMetaData->SetTakeNumber(NextTakeNumber);
-	}
-}
-
-void STakeRecorderCockpit::OnRecordingInitialized(UTakeRecorder* Recorder)
-{
-	// Recache the meta-data here since we know that the sequence has probably changed as a result of the recording being started
-	CacheMetaData();
-
-	OnRecordingFinishedHandle = Recorder->OnRecordingFinished().AddSP(this, &STakeRecorderCockpit::OnRecordingFinished);
-}
-
-void STakeRecorderCockpit::OnRecordingFinished(UTakeRecorder* Recorder)
-{
-	if (TransientTakeMetaData)
-	{
-		// Increment the transient take meta data if necessary
-		int32 NextTakeNumber = UTakesCoreBlueprintLibrary::ComputeNextTakeNumber(TransientTakeMetaData->GetSlate());
-
-		if (TransientTakeMetaData->GetTakeNumber() != NextTakeNumber)
-		{
-			TransientTakeMetaData->SetTakeNumber(NextTakeNumber);
-		}
-
-		bAutoApplyTakeNumber = true;
-	}
-
-	Recorder->OnRecordingFinished().Remove(OnRecordingFinishedHandle);
-}
-
 void STakeRecorderCockpit::BindCommands()
 {
 	// Bind our commands to the play world so that we can record in editor and in PIE
@@ -1269,13 +1034,14 @@ void STakeRecorderCockpit::OnToggleEditPreviousRecording(ECheckBoxState CheckSta
 {
 	if (Reviewing())
 	{
+		UTakeMetaData* TakeMetaData = GetMetaDataChecked();
 		TakeMetaData->IsLocked() ? TakeMetaData->Unlock() : TakeMetaData->Lock();
 	}	
 }
 
 bool STakeRecorderCockpit::EditingMetaData() const
 {
-	return (!Reviewing() || !TakeMetaData->IsLocked());
+	return (!Reviewing() || !GetMetaDataChecked()->IsLocked());
 }
 
 TSharedRef<SWidget> STakeRecorderCockpit::MakeLockButton()
@@ -1284,20 +1050,20 @@ TSharedRef<SWidget> STakeRecorderCockpit::MakeLockButton()
 	.Style(FAppStyle::Get(), "ToggleButtonCheckbox")
 	.Padding(TakeRecorder::ButtonPadding)
 	.ToolTipText(LOCTEXT("Modify Slate", "Unlock to modify the slate information for this prior recording."))
-	.IsChecked_Lambda([this]() { return TakeMetaData->IsLocked() ? ECheckBoxState::Unchecked: ECheckBoxState::Checked; } )
+	.IsChecked_Lambda([this]() { return GetMetaDataChecked()->IsLocked() ? ECheckBoxState::Unchecked: ECheckBoxState::Checked; } )
 	.OnCheckStateChanged(this, &STakeRecorderCockpit::OnToggleEditPreviousRecording)
 	.Visibility_Lambda([this]() { return Reviewing() ? EVisibility::Visible : EVisibility::Collapsed; })
 	[
 		SNew(STextBlock)
 		.Justification(ETextJustify::Center)
 		.Font(FAppStyle::Get().GetFontStyle("FontAwesome.14"))
-		.Text_Lambda([this]() { return TakeMetaData->IsLocked() ? FEditorFontGlyphs::Lock : FEditorFontGlyphs::Unlock; } )
+		.Text_Lambda([this]() { return GetMetaDataChecked()->IsLocked() ? FEditorFontGlyphs::Lock : FEditorFontGlyphs::Unlock; } )
 	];
 }
 
 TSharedRef<SWidget> STakeRecorderCockpit::OnCreateMenu()
 {
-	ULevelSequence* Sequence = LevelSequenceAttribute.Get();
+	ULevelSequence* Sequence = TakeRecorderSubsystem->GetLevelSequence();
 	if (!Sequence || !Sequence->GetMovieScene())
 	{
 		return SNullWidget::NullWidget;
@@ -1365,32 +1131,15 @@ TSharedRef<SWidget> STakeRecorderCockpit::OnCreateMenu()
 	return MenuBuilder.MakeWidget();
 }
 
-void STakeRecorderCockpit::BindToNamingTokenEvents()
-{
-	if (GEngine && !NamingTokensData->TakeRecorderNamingTokens.IsValid())
-	{
-		NamingTokensData->TakeRecorderNamingTokens = GEngine->GetEngineSubsystem<UNamingTokensEngineSubsystem>()->GetNamingTokens(ITakeRecorderNamingTokensModule::GetTakeRecorderNamespace());
-		if (ensure(NamingTokensData->TakeRecorderNamingTokens.IsValid()))
-		{
-			const TStrongObjectPtr<UNamingTokens> TakeRecorderNamingTokensPin = NamingTokensData->TakeRecorderNamingTokens.Pin();
-			TakeRecorderNamingTokensPin->RegisterExternalTokens(NamingTokensData->NamingTokensExternalGuid);
-			TakeRecorderNamingTokensPin->GetOnPreEvaluateEvent().AddSP(this, &STakeRecorderCockpit::OnTakeRecorderNamingTokensPreEvaluate);
-		}
-	}
-}
-
-void STakeRecorderCockpit::UnbindNamingTokensEvents()
+UTakeRecorderNamingTokensData* STakeRecorderCockpit::GetNamingTokensData() const
 {
-	if (const TStrongObjectPtr<UNamingTokens> TakeRecorderNamingTokensPin = NamingTokensData->TakeRecorderNamingTokens.Pin())
-	{
-		TakeRecorderNamingTokensPin->UnregisterExternalTokens(NamingTokensData->NamingTokensExternalGuid);
-		TakeRecorderNamingTokensPin->GetOnPreEvaluateEvent().RemoveAll(this);
-	}
+	check(TakeRecorderSubsystem.IsValid());
+	return TakeRecorderSubsystem->GetNamingTokensData();
 }
 
 FText STakeRecorderCockpit::GetCustomTokenTextValue(FString InTokenKey) const
 {
-	if (const FText* Value = NamingTokensData->UserDefinedTokens.Find(FNamingTokenData(InTokenKey)))
+	if (const FText* Value = GetNamingTokensData()->UserDefinedTokens.Find(FNamingTokenData(InTokenKey)))
 	{
 		return *Value;
 	}
@@ -1402,9 +1151,9 @@ void STakeRecorderCockpit::SetCustomTokenTextValue(const FText& InNewText, EText
 {
 	FScopedTransaction Transaction(LOCTEXT("SetCustomTokenTextValue_Transaction", "Set Custom Token Text"));
 	const FNamingTokenData CustomToken(InTokenKey);
-	NamingTokensData->Modify();
-	NamingTokensData->UserDefinedTokens.FindOrAdd(CustomToken) = InNewText;
-	NamingTokensData->VisibleUserTokens.Add(CustomToken.TokenKey);
+	GetNamingTokensData()->Modify();
+	GetNamingTokensData()->UserDefinedTokens.FindOrAdd(CustomToken) = InNewText;
+	GetNamingTokensData()->VisibleUserTokens.Add(CustomToken.TokenKey);
 	OnTokenValueUpdated();
 }
 
@@ -1421,32 +1170,11 @@ void STakeRecorderCockpit::EvaluateTakeSaveDirTokens()
 	}
 }
 
-void STakeRecorderCockpit::OnTakeRecorderNamingTokensPreEvaluate(const FNamingTokensEvaluationData& InEvaluationData)
-{
-	if (NamingTokensData->TakeRecorderNamingTokens.IsValid())
-	{
-		TArray<FNamingTokenData>& ExternalTokens = NamingTokensData->TakeRecorderNamingTokens->GetExternalTokensChecked(NamingTokensData->NamingTokensExternalGuid);
-
-		ExternalTokens.Reset(NamingTokensData->UserDefinedTokens.Num());
-		for (const TTuple<FNamingTokenData, FText>& UserToken : NamingTokensData->UserDefinedTokens)
-		{
-			ExternalTokens.Add({
-				UserToken.Key.TokenKey,
-				UserToken.Key.DisplayName,
-				// @todo naming tokens - If we support TOptional<FText> that can be used instead of a lambda
-				FNamingTokenData::FTokenProcessorDelegateNative::CreateLambda([Value = UserToken.Value]() {
-					return Value;
-				})
-			});
-		}
-	}
-}
-
 void STakeRecorderCockpit::RefreshUndefinedTokens()
 {
 	// Clicking on different actors will remove the old details panels -- clear out the fields
 	// we are tracking. They will be repopulated when evaluating the tokens below.
-	NamingTokensData->FieldToUndefinedKeys.Empty();
+	GetNamingTokensData()->FieldToUndefinedKeys.Empty();
 	
 	bool bHasMinRequiredObjects = false;
 	const UTakeRecorderUserSettings* UserSettings = GetDefault<UTakeRecorderUserSettings>();
@@ -1643,7 +1371,7 @@ void STakeRecorderCockpit::EvaluateTokensFromProperty(const FProperty* InPropert
 		const FString FieldName = FString::Printf(TEXT("%s_%p"), *InProperty->GetName(), InContainer);
 
 		// We will automatically inject our tokens to the take recorder naming tokens via our pre-evaluate hook.
-		NamingTokenResultData = TakeMetaData->ProcessTokens(ResultText, nullptr);
+		NamingTokenResultData = GetMetaDataChecked()->ProcessTokens(ResultText, nullptr);
 		
 		// Manually handle after all evaluation has finished. We don't need to listen for a post-evaluate event
 		// because we want the finalized text processed for all tokens classes.
@@ -1652,7 +1380,7 @@ void STakeRecorderCockpit::EvaluateTokensFromProperty(const FProperty* InPropert
 		for (const FNamingTokenValueData& Token : NamingTokenResultData.TokenValues)
 		{
 			if ((!Token.bWasEvaluated
-				|| NamingTokensData->UserDefinedTokens.Contains(FNamingTokenData(Token.TokenKey)))
+				|| GetNamingTokensData()->UserDefinedTokens.Contains(FNamingTokenData(Token.TokenKey)))
 					// Include pre-defined take and slate keys so they will get sorted with user tokens.
 				|| (Token.TokenKey == UE::TakeRecorder::Private::TokenKeyTake
 					|| Token.TokenKey == UE::TakeRecorder::Private::TokenKeySlate))
@@ -1664,7 +1392,7 @@ void STakeRecorderCockpit::EvaluateTokensFromProperty(const FProperty* InPropert
 		if (!FieldName.IsEmpty())
 		{
 			// Track the unidentified keys for this field.
-			TArray<FString>& TokenKeysForField = NamingTokensData->FindOrAddTokenKeysForField(FieldName);
+			TArray<FString>& TokenKeysForField = GetNamingTokensData()->FindOrAddTokenKeysForField(FieldName);
 			TokenKeysForField = UndefinedTokenKeys;
 		}
 	}
@@ -1674,8 +1402,8 @@ void STakeRecorderCockpit::EvaluateTokensFromProperty(const FProperty* InPropert
 	{
 		CachedTakeSaveDirProperty = InProperty;
 		CachedTakeSaveDirContainer = InContainer;
-		NamingTokensData->Modify();
-		NamingTokensData->EvaluatedTextValue = NamingTokenResultData.EvaluatedText;
+		GetNamingTokensData()->Modify();
+		GetNamingTokensData()->EvaluatedTextValue = NamingTokenResultData.EvaluatedText;
 	}
 }
 
@@ -1684,31 +1412,31 @@ void STakeRecorderCockpit::CreateUserTokensUI()
 	// Assemble the custom tokens in use. This step needs to happen after NamingTokensData->FieldToUndefinedKeys has
 	// been fully populated so we don't remove custom tokens that are still in use.
 	{
-		const TMap<FNamingTokenData, FText> CustomUserTokensCopy = NamingTokensData->UserDefinedTokens;
+		const TMap<FNamingTokenData, FText> CustomUserTokensCopy = GetNamingTokensData()->UserDefinedTokens;
 		for (const TTuple<FNamingTokenData, FText>& CustomToken : CustomUserTokensCopy)
 		{
 			// If no field holds a reference to this key then we can remove it.
-			if (!NamingTokensData->IsTokenKeyUndefined(CustomToken.Key.TokenKey))
+			if (!GetNamingTokensData()->IsTokenKeyUndefined(CustomToken.Key.TokenKey))
 			{
 				// Only remove the visible entry, but leave the value in CustomUserTokens. If the user is switching
 				// between objects it's possible one object has fields with different tokens than another, and we don't
 				// want to lose the old entered data.
-				NamingTokensData->VisibleUserTokens.Remove(CustomToken.Key.TokenKey);
+				GetNamingTokensData()->VisibleUserTokens.Remove(CustomToken.Key.TokenKey);
 			}
 		}
 
 		// Look for new unidentified tokens that a user can define.
-		for (const FTakeRecorderNamingTokensFieldMapping& UndefinedTokens : NamingTokensData->FieldToUndefinedKeys)
+		for (const FTakeRecorderNamingTokensFieldMapping& UndefinedTokens : GetNamingTokensData()->FieldToUndefinedKeys)
 		{
 			for (const FString& UndefinedTokenStr : UndefinedTokens.UndefinedKeys)
 			{
 				FNamingTokenData UndefinedToken(UndefinedTokenStr);
-				if (!NamingTokensData->UserDefinedTokens.Contains(UndefinedToken))
+				if (!GetNamingTokensData()->UserDefinedTokens.Contains(UndefinedToken))
 				{
 					// Check for existence so we don't overwrite the user defined value.
-					NamingTokensData->UserDefinedTokens.Add(UndefinedToken);
+					GetNamingTokensData()->UserDefinedTokens.Add(UndefinedToken);
 				}
-				NamingTokensData->VisibleUserTokens.Add(UndefinedToken.TokenKey);
+				GetNamingTokensData()->VisibleUserTokens.Add(UndefinedToken.TokenKey);
 			}
 		}
 	}
@@ -1720,11 +1448,11 @@ void STakeRecorderCockpit::CreateUserTokensUI()
 	UserTokensBox->ClearChildren();
 	
 	// Display tokens in order from their field appearance, to the token location in the string, but do not duplicate.
-	for (const FTakeRecorderNamingTokensFieldMapping& FieldMapping : NamingTokensData->FieldToUndefinedKeys)
+	for (const FTakeRecorderNamingTokensFieldMapping& FieldMapping : GetNamingTokensData()->FieldToUndefinedKeys)
 	{
 		for (const FString& TokenString : FieldMapping.UndefinedKeys)
 		{
-			if (UsedTokens.Contains(TokenString) || !NamingTokensData->VisibleUserTokens.Contains(TokenString))
+			if (UsedTokens.Contains(TokenString) || !GetNamingTokensData()->VisibleUserTokens.Contains(TokenString))
 			{
 				// Skip if this was already reference in this field or a previous one.
 				continue;
diff --git a/Engine/Plugins/VirtualProduction/Takes/Source/TakeRecorder/Private/Widgets/STakeRecorderCockpit.h b/Engine/Plugins/VirtualProduction/Takes/Source/TakeRecorder/Private/Widgets/STakeRecorderCockpit.h
--- a/Engine/Plugins/VirtualProduction/Takes/Source/TakeRecorder/Private/Widgets/STakeRecorderCockpit.h
+++ b/Engine/Plugins/VirtualProduction/Takes/Source/TakeRecorder/Private/Widgets/STakeRecorderCockpit.h
@@ -3,12 +3,12 @@
 #pragma once
 
 #include "Recorder/TakeRecorderParameters.h"
-#include "UObject/GCObject.h"
 #include "UObject/WeakFieldPtr.h"
 #include "Widgets/DeclarativeSyntaxSupport.h"
 #include "Widgets/SCompoundWidget.h"
-#include "Widgets/TakeRecorderNamingTokensWidgetData.h"
 
+class UTakeRecorderNamingTokensData;
+class UTakeRecorderSubsystem;
 enum class ECheckBoxState : uint8;
 enum class ETakeRecordMode : uint8;
 
@@ -28,28 +28,22 @@ struct FPropertyChangedEvent;
  * Cockpit UI for defining take meta-data.
  * Interacts with UTakeMetaData stored on the level sequence, if present, otherwise uses its own transient meta-data
  */
-class STakeRecorderCockpit : public SCompoundWidget, public FGCObject
+class STakeRecorderCockpit : public SCompoundWidget
 {
 public:
 
 	SLATE_BEGIN_ARGS(STakeRecorderCockpit)
-		: _LevelSequence(nullptr)
+		
 		{}
 
-		SLATE_ATTRIBUTE(ULevelSequence*, LevelSequence)
-
-		SLATE_ATTRIBUTE(ETakeRecorderMode, TakeRecorderMode)
-
 	SLATE_END_ARGS()
 
 	void Construct(const FArguments& InArgs);
 
 	~STakeRecorderCockpit();
 
-	UTakeMetaData* GetMetaData() const
-	{
-		return TakeMetaData;
-	}
+	UTakeMetaData* GetMetaData() const;
+	UTakeMetaData* GetMetaDataChecked() const;
 
 	FFrameRate GetFrameRate() const;
 
@@ -79,12 +73,6 @@ public:
 	void NotifyDetailsViewAdded(const TWeakPtr<IDetailsView>& InDetailsView);
 
 private:
-
-	virtual void AddReferencedObjects(FReferenceCollector& Collector) override;
-	virtual FString GetReferencerName() const override
-	{
-		return TEXT("STakeRecorderCockpit");
-	}
 	virtual void Tick(const FGeometry& AllottedGeometry, const double InCurrentTime, const float InDeltaTime) override;
 
 	void PostUndoRedo();
@@ -154,25 +142,15 @@ private:
 
 	bool EditingMetaData() const;
 
-	void CacheMetaData();
-
-	void OnAssetRegistryFilesLoaded();
-
-	void OnRecordingInitialized(UTakeRecorder*);
-
-	void OnRecordingFinished(UTakeRecorder*);
-
 	void BindCommands();
 
 	void OnToggleEditPreviousRecording(ECheckBoxState CheckState);
 	
 	TSharedRef<SWidget> OnCreateMenu();
 
-	/** Setup bindings to our naming tokens. */
-	void BindToNamingTokenEvents();
-	/** Remove bindings to our naming tokens. */
-	void UnbindNamingTokensEvents();
-
+	/** Retrieve data relating to naming tokens. */
+	UTakeRecorderNamingTokensData* GetNamingTokensData() const;
+	
 	/** Find a user defined text value for our custom tokens. */
 	FText GetCustomTokenTextValue(FString InTokenKey) const;
 	/** Set a user defined text value for a custom token. */
@@ -184,9 +162,6 @@ private:
 	/** Evaluate the take dir tokens for our example display. */
 	void EvaluateTakeSaveDirTokens();
 	
-	/** Called before our naming tokens class evaluates. Used to populate naming token data with our custom token definitions. */
-	void OnTakeRecorderNamingTokensPreEvaluate(const FNamingTokensEvaluationData& InEvaluationData);
-
 	/** Iterate all cached details views refreshing all undefined tokens. */
 	void RefreshUndefinedTokens();
 	
@@ -209,31 +184,15 @@ private:
 	
 private:
 
-	/** Take meta-data cached from the level sequence if it exists. Referenced by AddReferencedObjects. */
-	TObjectPtr<UTakeMetaData> TakeMetaData;
-
-	/** Transient take meta data owned by this widget and kept alive by AddReferencedObjects. Only used if none exists on the level sequence already. */
-	TObjectPtr<UTakeMetaData> TransientTakeMetaData;
-
 	/** The index of a pending transaction initiated by this widget, or INDEX_NONE if none is pending */
 	int32 TransactionIndex;
 
-	TAttribute<ULevelSequence*> LevelSequenceAttribute;
-
-	TAttribute<ETakeRecorderMode> TakeRecorderModeAttribute;
-
 	/** Text that describes why the user cannot record with the current settings */
 	FText RecordErrorText;
 
 	/** Text that describes why the user cannot record with the current settings */
 	FText TakeErrorText;
 
-	/** Whether we should auto apply the next available take number when asset discovery has finished or not */
-	bool bAutoApplyTakeNumber;
-
-	FDelegateHandle OnAssetRegistryFilesLoadedHandle;
-	FDelegateHandle OnRecordingInitializedHandle, OnRecordingFinishedHandle;
-
 	TSharedPtr<FDigitsTypeInterface> DigitsTypeInterface;
 
 	/** All known detail views. */
@@ -257,8 +216,8 @@ private:
 	/** Our cached object container, holding our Take Save Dir property value. */
 	const void* CachedTakeSaveDirContainer = nullptr;
 
-	/** Stored data relevant to our naming tokens. Managed as UObject for transactions. */
-	TStrongObjectPtr<UTakeRecorderNamingTokensWidgetData> NamingTokensData;
+	/** Weak ptr to take recorder subsystem. */
+	TWeakObjectPtr<UTakeRecorderSubsystem> TakeRecorderSubsystem;
 
 	/** Handle for binding to the post undo/redo editor delegate. */
 	FDelegateHandle PostUndoRedoDelegateHandle;
diff --git a/Engine/Plugins/VirtualProduction/Takes/Source/TakeRecorder/Private/Widgets/STakeRecorderPanel.cpp b/Engine/Plugins/VirtualProduction/Takes/Source/TakeRecorder/Private/Widgets/STakeRecorderPanel.cpp
--- a/Engine/Plugins/VirtualProduction/Takes/Source/TakeRecorder/Private/Widgets/STakeRecorderPanel.cpp
+++ b/Engine/Plugins/VirtualProduction/Takes/Source/TakeRecorder/Private/Widgets/STakeRecorderPanel.cpp
@@ -66,83 +66,39 @@
 #include "SequencerSettings.h"
 
 #include "LevelEditor.h"
+#include "Recorder/TakeRecorderSubsystem.h"
 
 #define LOCTEXT_NAMESPACE "STakeRecorderPanel"
 
 STakeRecorderPanel::~STakeRecorderPanel()
 {
-	UTakeRecorder::OnRecordingInitialized().Remove(OnRecordingInitializedHandle);
 	ITakeRecorderModule& TakeRecorderModule = FModuleManager::Get().LoadModuleChecked<ITakeRecorderModule>("TakeRecorder");
 	TakeRecorderModule.OnForceSaveAsPreset().Unbind();
-	TakeRecorderModule.GetLastLevelSequenceProvider().Unbind();
 	TakeRecorderModule.GetExternalObjectAddRemoveEventDelegate().Remove(OnWidgetExternalObjectChangedHandle);
+	
+	if (TakeRecorderSubsystem.IsValid())
+	{
+		TakeRecorderSubsystem->GetOnRecordingInitializedEvent().Remove(OnRecordingInitializedHandle);
+		TakeRecorderSubsystem->GetOnRecordingFinishedEvent().Remove(OnRecordingFinishedHandle);
+		TakeRecorderSubsystem->GetOnRecordingCancelledEvent().Remove(OnRecordingCancelledHandle);
+	}
 }
 
 UE_DISABLE_OPTIMIZATION_SHIP
 void STakeRecorderPanel::Construct(const FArguments& InArgs)
 {
-	// If a recording is currently underway, initialize to that now
-	if (UTakeRecorder* ActiveRecorder = UTakeRecorder::GetActiveRecorder())
-	{
-		RecordingLevelSequence = ActiveRecorder->GetSequence();
-		OnRecordingFinishedHandle = ActiveRecorder->OnRecordingFinished().AddSP(this, &STakeRecorderPanel::OnRecordingFinished);
-		OnRecordingCancelledHandle = ActiveRecorder->OnRecordingCancelled().AddSP(this, &STakeRecorderPanel::OnRecordingCancelled);
-	}
-	else
-	{
-		RecordingLevelSequence = nullptr;
-	}
-
-	TransientPreset = AllocateTransientPreset();
-	LastRecordedLevelSequence = nullptr;
-
-	// Copy the base preset into the transient preset if it was provided.
-	// We do this first so that anything that asks for its Level Sequence
-	// on construction gets the right one
-	if (InArgs._BasePreset)
-	{
-		TransientPreset->CopyFrom(InArgs._BasePreset);
-	}
-	else if (InArgs._BaseSequence)
-	{
-		TransientPreset->CopyFrom(InArgs._BaseSequence);
-
-		ULevelSequence* LevelSequence = TransientPreset->GetLevelSequence();
-
-		UTakeRecorderSources* BaseSources = InArgs._BaseSequence->FindMetaData<UTakeRecorderSources>();
-		if (BaseSources && LevelSequence)
-		{
-			LevelSequence->CopyMetaData(BaseSources);
-		}
-
-		if (LevelSequence)
-		{
-			LevelSequence->GetMovieScene()->SetReadOnly(false);
-		}
-
-		UTakeMetaData*  TakeMetaData = LevelSequence ? LevelSequence->FindMetaData<UTakeMetaData>() : nullptr;
-		if (TakeMetaData)
-		{
-			TakeMetaData->Unlock();
-			TakeMetaData->SetTimestamp(FDateTime(0));
-		}
-	}
-	else if (InArgs._RecordIntoSequence)
-	{
-		SetRecordIntoLevelSequence(InArgs._RecordIntoSequence);
-	}
-	else if (InArgs._SequenceToView)
-	{
-		SuppliedLevelSequence  = InArgs._SequenceToView;
-	}
-
-	/** Clear the dirty flag since the preset was just initialized. */
-	TransientPreset->GetOutermost()->SetDirtyFlag(false);
-
+	TakeRecorderSubsystem = GEngine->GetEngineSubsystem<UTakeRecorderSubsystem>();
+	
+	FTakeRecorderSequenceParameters Data;
+	Data.BasePreset = InArgs._BasePreset;
+	Data.BaseSequence = InArgs._BaseSequence;
+	Data.RecordIntoSequence = InArgs._RecordIntoSequence;
+	Data.SequenceToView = InArgs._SequenceToView;
+	
+	TakeRecorderSubsystem->SetTargetSequence(Data);
+	
 	// Create the child widgets that need to know about our level sequence
-	CockpitWidget = SNew(STakeRecorderCockpit)
-	.LevelSequence(this, &STakeRecorderPanel::GetLevelSequence)
-	.TakeRecorderMode(this, &STakeRecorderPanel::GetTakeRecorderMode);
+	CockpitWidget = SNew(STakeRecorderCockpit);
 
 	LevelSequenceTakeWidget = SNew(SLevelSequenceTakeEditor)
 	.LevelSequence(this, &STakeRecorderPanel::GetLevelSequence)
@@ -153,15 +109,16 @@ void STakeRecorderPanel::Construct(const FArguments& InArgs)
 	SequencerPanel = MakeShared<FScopedSequencerPanel>(MakeAttributeSP(this, &STakeRecorderPanel::GetLevelSequence));
 
 	// Bind onto the necessary delegates we need
-	OnLevelSequenceChangedHandle = TransientPreset->AddOnLevelSequenceChanged(FSimpleDelegate::CreateSP(this, &STakeRecorderPanel::OnLevelSequenceChanged));
-	OnRecordingInitializedHandle = UTakeRecorder::OnRecordingInitialized().AddSP(this, &STakeRecorderPanel::OnRecordingInitialized);
+	OnLevelSequenceChangedHandle = TakeRecorderSubsystem->GetTransientPreset()->AddOnLevelSequenceChanged(FSimpleDelegate::CreateSP(this, &STakeRecorderPanel::OnLevelSequenceChanged));
+	OnRecordingInitializedHandle = TakeRecorderSubsystem->GetOnRecordingInitializedEvent().AddSP(this, &STakeRecorderPanel::OnRecordingInitialized);
+	OnRecordingFinishedHandle = TakeRecorderSubsystem->GetOnRecordingFinishedEvent().AddSP(this, &STakeRecorderPanel::OnRecordingFinished);
+	OnRecordingCancelledHandle = TakeRecorderSubsystem->GetOnRecordingCancelledEvent().AddSP(this, &STakeRecorderPanel::OnRecordingCancelled);
 
 	ITakeRecorderModule& TakeRecorderModule = FModuleManager::Get().LoadModuleChecked<ITakeRecorderModule>("TakeRecorder");
 	OnWidgetExternalObjectChangedHandle =
 		TakeRecorderModule.GetExternalObjectAddRemoveEventDelegate().AddSP(this, &STakeRecorderPanel::ReconfigureExternalSettings);
 
 	TakeRecorderModule.OnForceSaveAsPreset().BindRaw(this, &STakeRecorderPanel::OnSaveAsPreset);
-	TakeRecorderModule.GetLastLevelSequenceProvider().BindRaw(this, &STakeRecorderPanel::SetLastLevelSequence);
 
 	for(TWeakObjectPtr<> Object : TakeRecorderModule.GetExternalObjects())
 	{
@@ -174,7 +131,7 @@ void STakeRecorderPanel::Construct(const FArguments& InArgs)
 	// Setup the preset origin for the meta-data in the cockpit if one was supplied
 	if (InArgs._BasePreset)
 	{
-		CockpitWidget->GetMetaData()->SetPresetOrigin(InArgs._BasePreset);
+		CockpitWidget->GetMetaDataChecked()->SetPresetOrigin(InArgs._BasePreset);
 	}
 
 	// Add the settings immediately if the user preference tells us to
@@ -537,105 +494,45 @@ UE_ENABLE_OPTIMIZATION_SHIP
 
 ULevelSequence* STakeRecorderPanel::GetLevelSequence() const
 {
-	if (SuppliedLevelSequence)
-	{
-		return SuppliedLevelSequence;
-	}
-	else if (RecordIntoLevelSequence)
-	{
-		return RecordIntoLevelSequence;
-	}
-	else if (RecordingLevelSequence)
-	{
-		return RecordingLevelSequence;
-	}
-	else
-	{
-		return TransientPreset->GetLevelSequence();
-	}
+	return TakeRecorderSubsystem->GetLevelSequence();
 }
 
 ULevelSequence* STakeRecorderPanel::GetLastRecordedLevelSequence() const
 {
-	return LastRecordedLevelSequence;
-}
-
-void STakeRecorderPanel::SetLastLevelSequence(ULevelSequence *InSequence)
-{
-	LastRecordedLevelSequence = InSequence;
+	return TakeRecorderSubsystem->GetLastRecordedLevelSequence();
 }
 
 bool STakeRecorderPanel::CanReviewLastLevelSequence() const
 {
-	if (LastRecordedLevelSequence == nullptr)
-	{
-		return false;
-	}
-	ITakeRecorderModule& TakeRecorderModule = FModuleManager::Get().LoadModuleChecked<ITakeRecorderModule>("TakeRecorder");
-	FCanReviewLastRecordedLevelSequence& CanReview = TakeRecorderModule.GetCanReviewLastRecordedLevelSequenceDelegate();
-	if (CanReview.IsBound())
-	{
-		return CanReview.Execute();
-	}
-	
-	return true;
+	return TakeRecorderSubsystem->CanReviewLastRecording();
 }
 
 ETakeRecorderMode STakeRecorderPanel::GetTakeRecorderMode() const
 {
-	if (RecordIntoLevelSequence != nullptr)
-	{
-		return ETakeRecorderMode::RecordIntoSequence;
-	}
-
-	return ETakeRecorderMode::RecordNewSequence;
+	return TakeRecorderSubsystem->GetTakeRecorderMode();
 }
 
 UTakeMetaData* STakeRecorderPanel::GetTakeMetaData() const
 {
-	return CockpitWidget->GetMetaData();
+	return CockpitWidget->GetMetaDataChecked();
 }
 
 void STakeRecorderPanel::ClearPendingTake()
 {
-	if (CockpitWidget->Reviewing())
-	{
-		LastRecordedLevelSequence = SuppliedLevelSequence;
-	}
-
-	UTakeRecorderSources* BaseSources = nullptr;
-
-	if (ULevelSequence* CurrentLevelSequence = GetLevelSequence())
-	{
-		BaseSources = CurrentLevelSequence->FindMetaData<UTakeRecorderSources>();
-	}
-
-	SuppliedLevelSequence = nullptr;
-	RecordIntoLevelSequence = nullptr;
-
-	FScopedTransaction Transaction(LOCTEXT("ClearPendingTake_Transaction", "Clear Pending Take"));
-
-	TransientPreset->Modify();
-	TransientPreset->CreateLevelSequence();
-
-	ULevelSequence* LevelSequence = TransientPreset->GetLevelSequence();
-	if (LevelSequence && BaseSources)
-	{
-		LevelSequence->CopyMetaData(BaseSources);
-	}
+	TakeRecorderSubsystem->ClearPendingTake();
 }
 
 TOptional<ETakeRecorderPanelMode> STakeRecorderPanel::GetMode() const
 {
-	if (SuppliedLevelSequence)
+	if (TakeRecorderSubsystem->GetSuppliedLevelSequence())
 	{
 		return ETakeRecorderPanelMode::ReviewingRecording;	
 	}
-	else if (RecordingLevelSequence)
+	else if (TakeRecorderSubsystem->GetRecordingLevelSequence())
 	{
 		return ETakeRecorderPanelMode::NewRecording;	
 	}
-	else if (RecordIntoLevelSequence)
+	else if (TakeRecorderSubsystem->GetRecordIntoLevelSequence())
 	{
 		return ETakeRecorderPanelMode::RecordingInto;	
 	}
@@ -643,22 +540,6 @@ TOptional<ETakeRecorderPanelMode> STakeRecorderPanel::GetMode() const
 	return TOptional<ETakeRecorderPanelMode>();
 }
 
-UTakePreset* STakeRecorderPanel::AllocateTransientPreset()
-{
-	return UTakePreset::AllocateTransientPreset(GetDefault<UTakeRecorderUserSettings>()->LastOpenedPreset.Get());
-}
-
-
-void STakeRecorderPanel::AddReferencedObjects(FReferenceCollector& Collector)
-{
-	Collector.AddReferencedObject(TransientPreset);
-	Collector.AddReferencedObject(RecordIntoLevelSequence);
-	Collector.AddReferencedObject(SuppliedLevelSequence);
-	Collector.AddReferencedObject(RecordingLevelSequence);
-	Collector.AddReferencedObject(LastRecordedLevelSequence);
-}
-
-
 TSharedRef<SWidget> STakeRecorderPanel::OnGeneratePresetsMenu()
 {
 	FMenuBuilder MenuBuilder(true, nullptr);
@@ -715,24 +596,7 @@ void STakeRecorderPanel::OnImportPreset(const FAssetData& InPreset)
 {
 	FSlateApplication::Get().DismissAllMenus();
 
-	SuppliedLevelSequence = nullptr;
-	RecordIntoLevelSequence = nullptr;
-
-	UTakePreset* Take = CastChecked<UTakePreset>(InPreset.GetAsset());
-	if (Take)
-	{
-		FScopedTransaction Transaction(LOCTEXT("ImportPreset_Transaction", "Import Take Preset"));
-
-		TransientPreset->Modify();
-		TransientPreset->CopyFrom(Take);
-		TransientPreset->GetOutermost()->SetDirtyFlag(false);
-
-		CockpitWidget->GetMetaData()->SetPresetOrigin(Take);
-	}
-	else
-	{
-		// @todo: notification?
-	}
+	TakeRecorderSubsystem->ImportPreset(InPreset);
 }
 
 
@@ -831,6 +695,8 @@ void STakeRecorderPanel::OnSaveAsPreset()
 
 	if (NewPreset)
 	{
+		UTakePreset* TransientPreset = TakeRecorderSubsystem->GetTransientPreset();
+		
 		NewPreset->CopyFrom(TransientPreset);
 		if (ULevelSequence* LevelSequence = NewPreset->GetLevelSequence())
 		{
@@ -845,21 +711,13 @@ void STakeRecorderPanel::OnSaveAsPreset()
 
 		FEditorFileUtils::PromptForCheckoutAndSave({ NewPackage }, false, false);
 
-		CockpitWidget->GetMetaData()->SetPresetOrigin(NewPreset);
+		TakeRecorderSubsystem->GetTakeMetaData()->SetPresetOrigin(NewPreset);
 	}
 }
 
 FReply STakeRecorderPanel::OnBackToPendingTake()
 {
-	if (CockpitWidget->Reviewing())
-	{
-		LastRecordedLevelSequence = SuppliedLevelSequence;
-	}
-
-	SuppliedLevelSequence = nullptr;
-	RecordIntoLevelSequence = nullptr;
-	
-	TransientPreset = AllocateTransientPreset();
+	TakeRecorderSubsystem->ResetToPendingTake();
 	RefreshPanel();
 
 	return FReply::Handled();
@@ -879,10 +737,8 @@ FReply STakeRecorderPanel::OnClearPendingTake()
 
 FReply STakeRecorderPanel::OnReviewLastRecording()
 {
-	if (LastRecordedLevelSequence)
+	if (TakeRecorderSubsystem->ReviewLastRecording())
 	{
-		SuppliedLevelSequence = LastRecordedLevelSequence;
-		LastRecordedLevelSequence = nullptr;
 		RefreshPanel();
 	}
 
@@ -897,13 +753,7 @@ FReply STakeRecorderPanel::OnRevertChanges()
 		return FReply::Handled();
 	}
 
-	UTakePreset* PresetOrigin = CockpitWidget->GetMetaData()->GetPresetOrigin();
-
-	FScopedTransaction Transaction(LOCTEXT("RevertChanges_Transaction", "Revert Changes"));
-
-	TransientPreset->Modify();
-	TransientPreset->CopyFrom(PresetOrigin);
-	TransientPreset->GetOutermost()->SetDirtyFlag(false);
+	TakeRecorderSubsystem->RevertChanges();
 
 	return FReply::Handled();
 }
@@ -956,10 +806,9 @@ void STakeRecorderPanel::OnOpenSequenceToRecordInto(const FAssetData& InAsset)
 	FSlateApplication::Get().DismissAllMenus();
 
 	// Only try to initialize level sequences, in the event they had more than a level sequence selected when drag/dropping.
-	ULevelSequence* LevelSequence = Cast<ULevelSequence>(InAsset.GetAsset());
-	if (LevelSequence)
+	if (ULevelSequence* LevelSequence = Cast<ULevelSequence>(InAsset.GetAsset()))
 	{
-		SetRecordIntoLevelSequence(LevelSequence);
+		TakeRecorderSubsystem->SetRecordIntoLevelSequence(LevelSequence);
 		RefreshPanel();
 	}
 }
@@ -1020,62 +869,22 @@ void STakeRecorderPanel::OnLevelSequenceDetailsViewAdded(const TWeakPtr<IDetails
 	}
 }
 
-void STakeRecorderPanel::SetRecordIntoLevelSequence(ULevelSequence* LevelSequence)
-{
-	SuppliedLevelSequence = nullptr;
-	RecordIntoLevelSequence = LevelSequence;
-
-	if (RecordIntoLevelSequence)
-	{
-		RecordIntoLevelSequence->GetMovieScene()->SetReadOnly(false);
-	}
-
-	UTakeMetaData*  TakeMetaData = RecordIntoLevelSequence ? RecordIntoLevelSequence->FindOrAddMetaData<UTakeMetaData>() : nullptr;
-	if (TakeMetaData)
-	{
-		// Set up take metadata to match this level sequence's info, ie. match the frame rate, use the level sequence name as the slate
-		TakeMetaData->Unlock();
-		TakeMetaData->SetTimestamp(FDateTime(0));
-		TakeMetaData->SetSlate(LevelSequence->GetName());
-		TakeMetaData->SetTakeNumber(0);
-		TakeMetaData->SetFrameRate(LevelSequence->GetMovieScene()->GetDisplayRate());
-		TakeMetaData->SetFrameRateFromTimecode(false);
-	}
-}
-
 void STakeRecorderPanel::OnRecordingInitialized(UTakeRecorder* Recorder)
 {
-	RecordingLevelSequence = Recorder->GetSequence();
+	// It's important that UTakeRecorderEditorSubsystem::OnRecordingInitialized has fired before this point,
+	// otherwise the refresh may invalidate a weak sequencer the subsystem relies on.
 	RefreshPanel();
-
-	OnRecordingFinishedHandle = Recorder->OnRecordingFinished().AddSP(this, &STakeRecorderPanel::OnRecordingFinished);
-	OnRecordingCancelledHandle = Recorder->OnRecordingCancelled().AddSP(this, &STakeRecorderPanel::OnRecordingCancelled);
 }
 
 void STakeRecorderPanel::OnRecordingFinished(UTakeRecorder* Recorder)
 {
-	LastRecordedLevelSequence = RecordingLevelSequence;
 	OnRecordingCancelled(Recorder);
-
-	// Update the preset take number at the end of recording
-	ULevelSequence* LevelSequence = TransientPreset->GetLevelSequence();
-	UTakeMetaData*  TakeMetaData = LevelSequence ? LevelSequence->FindMetaData<UTakeMetaData>() : nullptr;
-
-	if (TakeMetaData)
-	{
-		int32 NextTakeNumber = UTakesCoreBlueprintLibrary::ComputeNextTakeNumber(TakeMetaData->GetSlate());
-		TakeMetaData->SetTakeNumber(NextTakeNumber);
-	}
 }
 
 void STakeRecorderPanel::OnRecordingCancelled(UTakeRecorder* Recorder)
 {
-	RecordingLevelSequence = nullptr;
 	RefreshPanel();
 	CockpitWidget->Refresh();
-
-	Recorder->OnRecordingFinished().Remove(OnRecordingFinishedHandle);
-	Recorder->OnRecordingCancelled().Remove(OnRecordingCancelledHandle);
 }
 
 ECheckBoxState STakeRecorderPanel::GetTakeBrowserCheckState() const
diff --git a/Engine/Plugins/VirtualProduction/Takes/Source/TakeRecorder/Private/Widgets/STakeRecorderPanel.h b/Engine/Plugins/VirtualProduction/Takes/Source/TakeRecorder/Private/Widgets/STakeRecorderPanel.h
--- a/Engine/Plugins/VirtualProduction/Takes/Source/TakeRecorder/Private/Widgets/STakeRecorderPanel.h
+++ b/Engine/Plugins/VirtualProduction/Takes/Source/TakeRecorder/Private/Widgets/STakeRecorderPanel.h
@@ -24,11 +24,12 @@ class ULevelSequence;
 class UTakeRecorderSource;
 class STakeRecorderCockpit;
 class SLevelSequenceTakeEditor;
+class UTakeRecorderSubsystem;
 
 /**
  * Outermost widget that is used for setting up a new take recording. Operates on a transient UTakePreset that is internally owned and maintained 
  */
-class STakeRecorderPanel : public SCompoundWidget, public FGCObject
+class STakeRecorderPanel : public SCompoundWidget
 {
 public:
 
@@ -74,21 +75,6 @@ public:
 
 	TOptional<ETakeRecorderPanelMode> GetMode() const;
 
-private:
-
-	virtual void AddReferencedObjects(FReferenceCollector& Collector) override;
-	virtual FString GetReferencerName() const override
-	{
-		return TEXT("STakeRecorderPanel");
-	}
-
-private:
-
-	/**
-	 * Allocate the preset required for interacting with this widget. Re-uses an existing preset if the panel has been previously opened.
-	 */
-	static UTakePreset* AllocateTransientPreset();
-
 private:
 
 	/**
@@ -144,23 +130,11 @@ private:
 
 	TSharedRef<SWidget> MakeToolBar();
 
-	void SetRecordIntoLevelSequence(ULevelSequence* LevelSequence);
-
 private:
 	bool CanReviewLastLevelSequence() const;
 
-	void SetLastLevelSequence(ULevelSequence* InLastLevelSequence);
-
-	/** The transient preset that we use to - kept alive by AddReferencedObjects */
-	TObjectPtr<UTakePreset> TransientPreset;
-
-	TObjectPtr<ULevelSequence> SuppliedLevelSequence;
-
-	TObjectPtr<ULevelSequence> RecordIntoLevelSequence;
-
-	TObjectPtr<ULevelSequence> RecordingLevelSequence;
-
-	TObjectPtr<ULevelSequence> LastRecordedLevelSequence;
+	/** Weak ptr to subsystem. */
+	TWeakObjectPtr<UTakeRecorderSubsystem> TakeRecorderSubsystem;
 
 	/** The main level sequence take editor widget */
 	TSharedPtr<SLevelSequenceTakeEditor> LevelSequenceTakeWidget;
diff --git a/Engine/Plugins/VirtualProduction/Takes/Source/TakeRecorder/Private/Widgets/TakeRecorderNamingTokensData.cpp b/Engine/Plugins/VirtualProduction/Takes/Source/TakeRecorder/Private/Widgets/TakeRecorderNamingTokensData.cpp
new file mode 100644
--- /dev/null
+++ b/Engine/Plugins/VirtualProduction/Takes/Source/TakeRecorder/Private/Widgets/TakeRecorderNamingTokensData.cpp
@@ -0,0 +1,30 @@
+// Copyright Epic Games, Inc. All Rights Reserved.
+
+#include "TakeRecorderNamingTokensData.h"
+
+TArray<FString>& UTakeRecorderNamingTokensData::FindOrAddTokenKeysForField(const FString& InFieldName)
+{
+	if (FTakeRecorderNamingTokensFieldMapping* Match = FieldToUndefinedKeys.FindByPredicate([&InFieldName](const FTakeRecorderNamingTokensFieldMapping& Fields)
+	{
+		return Fields.FieldName == InFieldName;
+	}))
+	{
+		return Match->UndefinedKeys;
+	}
+	FTakeRecorderNamingTokensFieldMapping Data;
+	Data.FieldName = InFieldName;
+	return FieldToUndefinedKeys.Add_GetRef(Data).UndefinedKeys;
+}
+
+bool UTakeRecorderNamingTokensData::IsTokenKeyUndefined(const FString& InTokenKey) const
+{
+	for (const FTakeRecorderNamingTokensFieldMapping& TokenFieldMapping : FieldToUndefinedKeys)
+	{
+		if (TokenFieldMapping.UndefinedKeys.Contains(InTokenKey))
+		{
+			return true;
+		}
+	}
+
+	return false;
+}
diff --git a/Engine/Plugins/VirtualProduction/Takes/Source/TakeRecorder/Private/Widgets/TakeRecorderNamingTokensData.h b/Engine/Plugins/VirtualProduction/Takes/Source/TakeRecorder/Private/Widgets/TakeRecorderNamingTokensData.h
new file mode 100644
--- /dev/null
+++ b/Engine/Plugins/VirtualProduction/Takes/Source/TakeRecorder/Private/Widgets/TakeRecorderNamingTokensData.h
@@ -0,0 +1,66 @@
+// Copyright Epic Games, Inc. All Rights Reserved.
+
+#pragma once
+
+#include "Internationalization/LocKeyFuncs.h"
+#include "NamingTokenData.h"
+
+#include "TakeRecorderNamingTokensData.generated.h"
+
+class UNamingTokens;
+
+USTRUCT()
+struct FTakeRecorderNamingTokensFieldMapping
+{
+	GENERATED_BODY()
+
+	/** Name of the property field. */
+	UPROPERTY()
+	FString FieldName;
+	
+	/** All undefined keys with this field. */
+	UPROPERTY()
+	TArray<FString> UndefinedKeys;
+};
+
+/**
+ * Data container object for Take Recorder Naming Tokens. Kept as UObject to assist with transactions.
+ */
+UCLASS()
+class UTakeRecorderNamingTokensData : public UObject
+{
+	GENERATED_BODY()
+public:
+	/** Guid to our managed tokens. */
+	FGuid NamingTokensExternalGuid;
+	
+	/**
+	 * Custom tokens entered in by the user, mapped to a user defined value.
+	 * We use FNamingTokenData rather than an FString for the key so we can support case sensitivity in our map.
+	 */
+	UPROPERTY()
+	TMap<FNamingTokenData, FText> UserDefinedTokens;
+
+	/**
+	 * User tokens that are currently visible. Kept as a separate property from UsedDefinedTokens we can persist user values
+	 * between selected sources that have different token entries in their fields.
+	 */
+	TSet<FString, FLocKeySetFuncs> VisibleUserTokens;
+	
+	/** Ordered array of FieldName to undefined token keys. These are present if a token evaluation failed to identify them. */
+	UPROPERTY(NonTransactional)
+	TArray<FTakeRecorderNamingTokensFieldMapping> FieldToUndefinedKeys;
+
+	/** Evaluated text to serve as an example. */
+	UPROPERTY()
+	FText EvaluatedTextValue;
+	
+	/** Pointer to our naming tokens object. */
+	TWeakObjectPtr<UNamingTokens> TakeRecorderNamingTokens;
+
+	/** Find or add an array given a field name. The array will be a list of token keys for this specific field. */
+	TArray<FString>& FindOrAddTokenKeysForField(const FString& InFieldName);
+
+	/** Checks if a token key is in our undefined list. */
+	bool IsTokenKeyUndefined(const FString& InTokenKey) const;
+};
\ No newline at end of file
diff --git a/Engine/Plugins/VirtualProduction/Takes/Source/TakeRecorder/Private/Widgets/TakeRecorderNamingTokensWidgetData.cpp b/Engine/Plugins/VirtualProduction/Takes/Source/TakeRecorder/Private/Widgets/TakeRecorderNamingTokensWidgetData.cpp
deleted file mode 100644
diff --git a/Engine/Plugins/VirtualProduction/Takes/Source/TakeRecorder/Private/Widgets/TakeRecorderNamingTokensWidgetData.h b/Engine/Plugins/VirtualProduction/Takes/Source/TakeRecorder/Private/Widgets/TakeRecorderNamingTokensWidgetData.h
deleted file mode 100644
diff --git a/Engine/Plugins/VirtualProduction/Takes/Source/TakeRecorder/Public/ITakeRecorderModule.h b/Engine/Plugins/VirtualProduction/Takes/Source/TakeRecorder/Public/ITakeRecorderModule.h
--- a/Engine/Plugins/VirtualProduction/Takes/Source/TakeRecorder/Public/ITakeRecorderModule.h
+++ b/Engine/Plugins/VirtualProduction/Takes/Source/TakeRecorder/Public/ITakeRecorderModule.h
@@ -7,7 +7,9 @@
 #include "Templates/SharedPointer.h"
 #include "UObject/WeakObjectPtrTemplates.h"
 
+class AActor;
 class FExtender;
+class UTakeRecorderSource;
 class UTakeRecorderSources;
 class UTakeMetaData;
 class UTakePreset;
@@ -57,7 +59,35 @@ public:
 	 * parameter indicates if the object was added. It is true if added and false if removed.
 	 */
 	DECLARE_MULTICAST_DELEGATE_TwoParams(FOnExternalObjectAddRemoveEvent, UObject *, bool);
-
+	
+	/** Relevant data for the sources module. */
+	struct FSourceExtensionData
+	{
+		/**
+		 * Delegate for sources module to bind to when adding an actor.
+		 */
+		DECLARE_DELEGATE_FourParams(FAddActorSource, UTakeRecorderSources* TakeRecorderSources, TArrayView<AActor* const> ActorsToRecord, bool bReduceKeys, bool bShowProgress);                      
+
+		/**
+		 * Delegate for sources module to bind to when removing an actor.
+		 */
+		DECLARE_DELEGATE_TwoParams(FRemoveActorSource, UTakeRecorderSources* TakeRecorderSources, TArrayView<AActor* const> ActorsToRemove);
+
+		/**
+		 * Delegate to retrieve the actor of a source.
+		 */
+		DECLARE_DELEGATE_RetVal_OneParam(AActor*, FGetSourceActor, UTakeRecorderSource* Source);
+		
+		/** An actor is being added as a source. */
+		FAddActorSource OnAddActorSource;
+
+		/** An actor is being removed as a source. */
+		FRemoveActorSource OnRemoveActorSource;
+
+		/** The actor belonging to a source should be retrieved. */
+		FGetSourceActor OnGetSourceActor;
+	};
+	
 	/** The name under which the take recorder tab is registered and invoked */
 	static FName TakeRecorderTabName;
 
@@ -81,6 +111,16 @@ public:
 	 */
 	virtual FDelegateHandle RegisterSourcesMenuExtension(const FOnExtendSourcesMenu& InExtension) = 0;
 
+	/**
+	 * Register generic extensions for sources. We only allow one extension as this should be handled by the sources module.
+	 */
+	virtual void RegisterSourcesExtension(const FSourceExtensionData& InData) = 0;
+
+	/**
+	 * Unregister the source extension.
+	 */
+	virtual void UnregisterSourcesExtension() = 0;
+
 	/**
 	 * Get the delegate for handling saving of preset data.
 	 */
diff --git a/Engine/Plugins/VirtualProduction/Takes/Source/TakeRecorder/Public/Recorder/TakeRecorder.h b/Engine/Plugins/VirtualProduction/Takes/Source/TakeRecorder/Public/Recorder/TakeRecorder.h
--- a/Engine/Plugins/VirtualProduction/Takes/Source/TakeRecorder/Public/Recorder/TakeRecorder.h
+++ b/Engine/Plugins/VirtualProduction/Takes/Source/TakeRecorder/Public/Recorder/TakeRecorder.h
@@ -27,6 +27,7 @@ namespace UE::MovieScene
 UENUM(BlueprintType)
 enum class ETakeRecorderState : uint8
 {
+	PreInitialization,
 	CountingDown,
 	PreRecord,
 	TickingAfterPre,
@@ -93,6 +94,12 @@ public:
 		return CountdownSeconds;
 	}
 
+	/**
+	 * Sets the current countdown time in seconds.
+	 */
+	UFUNCTION(BlueprintCallable, Category="Take Recorder")
+	void SetCountdown(float InSeconds);
+
 	/**
 	 * Access the sequence asset that this recorder is recording into
 	 */
@@ -124,7 +131,7 @@ public:
 	 * @param LevelSequenceBase   A level sequence to use as a base set-up for the take
 	 * @param Sources             The sources to record from
 	 * @param MetaData            Meta data to store on the recorded level sequence asset for this take
-	 * @param Parameters          Configurable parameters for this instance of the recorder
+	 * @param InParameters        Configurable parameters for this instance of the recorder
 	 * @param OutError            Error string to receive for context
 	 * @return True if the recording process was successfully initialized, false otherwise
 	 */
@@ -200,12 +207,12 @@ private:
 	/**
 	 * Create a new destination asset to record into based on the parameters
 	 */
-	bool CreateDestinationAsset(const TCHAR* AssetPathFormat, ULevelSequence* LevelSequenceBase, UTakeRecorderSources* Sources, UTakeMetaData* MetaData, FText* OutError);
+	bool CreateDestinationAsset(const TCHAR* AssetPathFormat, ULevelSequence* LevelSequenceBase, UTakeRecorderSources* Sources, UTakeMetaData* MetaData, FText* OutError, bool bOpenSequencer = true);
 
 	/**
 	 * Setup a existing asset to record into based on the parameters
 	 */
-	bool SetupDestinationAsset(const FTakeRecorderParameters& InParameters, ULevelSequence* LevelSequenceBase, UTakeRecorderSources* Sources, UTakeMetaData* MetaData, FText* OutError);
+	bool SetupDestinationAsset(const FTakeRecorderParameters& InParameters, ULevelSequence* LevelSequenceBase, UTakeRecorderSources* Sources, UTakeMetaData* MetaData, FText* OutError, bool bOpenSequencer = true);
 
 	/**
 	 * Discovers the source world to record from, and initializes it for recording
diff --git a/Engine/Plugins/VirtualProduction/Takes/Source/TakeRecorder/Public/Recorder/TakeRecorderParameters.h b/Engine/Plugins/VirtualProduction/Takes/Source/TakeRecorder/Public/Recorder/TakeRecorderParameters.h
--- a/Engine/Plugins/VirtualProduction/Takes/Source/TakeRecorder/Public/Recorder/TakeRecorderParameters.h
+++ b/Engine/Plugins/VirtualProduction/Takes/Source/TakeRecorder/Public/Recorder/TakeRecorderParameters.h
@@ -8,6 +8,8 @@
 #include "TrackRecorders/IMovieSceneTrackRecorderHost.h"
 #include "TakeRecorderParameters.generated.h"
 
+class UTakePreset;
+
 USTRUCT(BlueprintType)
 struct FTakeRecorderUserParameters
 {
@@ -170,4 +172,39 @@ struct FTakeRecorderParameters
 	 */
 	UPROPERTY(Transient)
 	bool bDisableRecordingAndSave = false;
+
+	/** If sequencer should open when starting a recording. */
+	UPROPERTY(BlueprintReadWrite, Category = "Take Recorder")
+	bool bOpenSequencer = true;
+};
+
+/**
+ * Parameters for initializing a new sequence.
+ */
+USTRUCT(Blueprintable)
+struct TAKERECORDER_API FTakeRecorderSequenceParameters
+{
+	GENERATED_BODY()
+	
+	/** A preset asset to base the recording off */
+	UPROPERTY(BlueprintReadWrite, Category = TakeRecorder)
+	TObjectPtr<UTakePreset> BasePreset;
+
+	/** A level sequence asset to base the recording off */
+	UPROPERTY(BlueprintReadWrite, Category = TakeRecorder)
+	TObjectPtr<ULevelSequence> BaseSequence;
+
+	/** A level sequence asset to record into */
+	UPROPERTY(BlueprintReadWrite, Category = TakeRecorder)
+	TObjectPtr<ULevelSequence> RecordIntoSequence;
+
+	/** A sequence that should be shown directly on the take recorder UI */
+	UPROPERTY(BlueprintReadWrite, Category = TakeRecorder)
+	TObjectPtr<ULevelSequence> SequenceToView;
+
+	/** Checks whether we have any valid data present. */
+	bool IsValid() const;
+	
+	bool operator==(const FTakeRecorderSequenceParameters& Other) const;
+	bool operator!=(const FTakeRecorderSequenceParameters& Other) const;
 };
diff --git a/Engine/Plugins/VirtualProduction/Takes/Source/TakeRecorder/Public/Recorder/TakeRecorderSubsystem.h b/Engine/Plugins/VirtualProduction/Takes/Source/TakeRecorder/Public/Recorder/TakeRecorderSubsystem.h
--- a/Engine/Plugins/VirtualProduction/Takes/Source/TakeRecorder/Public/Recorder/TakeRecorderSubsystem.h
+++ b/Engine/Plugins/VirtualProduction/Takes/Source/TakeRecorder/Public/Recorder/TakeRecorderSubsystem.h
@@ -3,50 +3,452 @@
 #pragma once
 
 #include "Subsystems/EngineSubsystem.h"
+#include "TakeRecorder.h"
+#include "TakeRecorderSources.h"
+#include "TickableEditorObject.h"
 
 #include "TakeRecorderSubsystem.generated.h"
 
+struct FNamingTokensEvaluationData;
 class ULevelSequence;
+class UTakeRecorderNamingTokensData;
 
 DECLARE_DYNAMIC_MULTICAST_DELEGATE(FTakeRecorderPreInitialize);
+DECLARE_DYNAMIC_MULTICAST_DELEGATE(FTakeRecorderInitialized);
 DECLARE_DYNAMIC_MULTICAST_DELEGATE(FTakeRecorderStarted);
 DECLARE_DYNAMIC_MULTICAST_DELEGATE(FTakeRecorderStopped);
 DECLARE_DYNAMIC_MULTICAST_DELEGATE_OneParam(FTakeRecorderFinished, ULevelSequence*, SequenceAsset);
 DECLARE_DYNAMIC_MULTICAST_DELEGATE(FTakeRecorderCancelled);
 DECLARE_DYNAMIC_MULTICAST_DELEGATE_OneParam(FTakeRecorderMarkedFrameAdded, const FMovieSceneMarkedFrame&, MarkedFrame);
 
+/**
+ * Dynamic delegates for source modification. These are reproductions of static ones declared under UTakeRecorderSource.
+ * We do this so we can successfully leverage static delegates within the TakesCore module, firing when any change is made for
+ * any source owner. We define new ones here so they can be blueprint assignable.
+ */
+DECLARE_DYNAMIC_MULTICAST_DELEGATE_OneParam(FOnTakeRecorderSourceAddedDynamic, UTakeRecorderSource*, Source);
+DECLARE_DYNAMIC_MULTICAST_DELEGATE_OneParam(FOnTakeRecorderSourceRemovedDynamic, UTakeRecorderSource*, Source);
+
 /**
 * UTakeRecorderSubsystem
 * Subsystem for Take Recorder
 */
 UCLASS()
-class TAKERECORDER_API UTakeRecorderSubsystem : public UEngineSubsystem
+class TAKERECORDER_API UTakeRecorderSubsystem : public UEngineSubsystem, public FTickableEditorObject
 {
 	GENERATED_BODY()
 
+public:
+	virtual void Initialize(FSubsystemCollectionBase& Collection) override;
+	virtual void Deinitialize() override;
+
+	// Begin FTickableEditorObject
+	virtual void Tick(float DeltaTime) override;
+	virtual TStatId GetStatId() const override;
+	// ~End FTickableEditorObject
+	
+	/** Provide sequence data for this take recorder. */
+	UFUNCTION(BlueprintCallable, Category = TakeRecorder)
+	void SetTargetSequence(const FTakeRecorderSequenceParameters& InData);
+	
+	/** Set the record into level sequence. */
+	void SetRecordIntoLevelSequence(ULevelSequence* LevelSequence);
+
+	/** Can we review the last recording? */
+	UFUNCTION(BlueprintCallable, Category = TakeRecorder)
+	bool CanReviewLastRecording() const;
+	
+	/**
+	 * Supply the last recording if it exists.
+	 * @return true if it can be reviewed, false if it there isn't a recording to review.
+	 */
+	UFUNCTION(BlueprintCallable, Category = TakeRecorder)
+	bool ReviewLastRecording();
+
+	/**
+	 * Begin a new recording.
+	 * @param bOpenSequencer If sequencer should open when starting the recording.
+	 * @param bShowErrorMessage If an error message should be displayed on failure.
+	 * @return true if the recording was started successfully.
+	 */
+	UFUNCTION(BlueprintCallable, Category = TakeRecorder)
+	bool StartRecording(bool bOpenSequencer = true, bool bShowErrorMessage = true);
+
+	/**
+	 * Stop an existing recording.
+	 */
+	UFUNCTION(BlueprintCallable, Category = TakeRecorder)
+	void StopRecording();
+
+	/**
+	 * Cancel an in-progress recording.
+	 */
+	UFUNCTION(BlueprintCallable, Category = TakeRecorder)
+	void CancelRecording(bool bShowConfirmMessage = true);
+
+	/**
+	 * Reset to the pending take.
+	 */
+	UFUNCTION(BlueprintCallable, Category = TakeRecorder)
+	void ResetToPendingTake();
+
+	/**
+	 * Clear the pending take.
+	 */
+	UFUNCTION(BlueprintCallable, Category = TakeRecorder)
+	void ClearPendingTake();
+
+	/** Retrieve the pending take. This may be null. */
+	UFUNCTION(BlueprintCallable, Category = TakeRecorder)
+	UTakePreset* GetPendingTake() const;
+
+	/** Revert any changes restoring the preset origin. */
+	UFUNCTION(BlueprintCallable, Category = TakeRecorder)
+	void RevertChanges();
+	
+	// ~Begin Sources
+	/** Add a source by a source class. */
+	UFUNCTION(BlueprintCallable, Category = TakeRecorder, meta = (DeterminesOutputType = "InSourceClass"))
+	UTakeRecorderSource* AddSource(const TSubclassOf<UTakeRecorderSource> InSourceClass);
+	
+	/** Remove a given source. */
+	UFUNCTION(BlueprintCallable, Category = TakeRecorder)
+	void RemoveSource(UTakeRecorderSource* InSource);
+
+	/** Remove all sources from the current sequence. */
+	UFUNCTION(BlueprintCallable, Category = TakeRecorder)
+	void ClearSources();
+
+	/** Retrieve the sources. */
+	UFUNCTION(BlueprintCallable, Category = TakeRecorder)
+	UTakeRecorderSources* GetSources() const;
+	
+	/** Retrieve all sources for the current sequence. */
+	TArrayView<UTakeRecorderSource* const> GetAllSources() const;
+
+	/**
+	 * Retrieves a copy of the list of sources that are being recorded. This is intended for Blueprint usages which cannot
+	 * use TArrayView.
+	 * DO NOT MODIFY THIS ARRAY, modifications will be lost.
+	 */
+	UFUNCTION(BlueprintPure, DisplayName = "Get Sources (Copy)", Category = TakeRecorder)
+	TArray<UTakeRecorderSource*> GetAllSourcesCopy() const;
+
+	/** Retrieve the first source of the given class. */
+	UFUNCTION(BlueprintCallable, Category = TakeRecorder, meta = (DeterminesOutputType = "InSourceClass"))
+	UTakeRecorderSource* GetSourceByClass(const TSubclassOf<UTakeRecorderSource> InSourceClass) const;
+
+	/*
+	 * Add an actor as a source.
+	 * 
+	 * @param InActor The actor that should be added to Sources. Note that this can include ALevelSequenceActors.
+	 * @param bReduceKeys Enable/disable key reduction on all the sources registered
+	 * @param bShowProgress Enable/disable the dialog box showing progress for the potentially slow parts of finalizing the take
+	 */
+	UFUNCTION(BlueprintCallable, Category = TakeRecorder)
+	void AddSourceForActor(AActor* InActor, bool bReduceKeys = true, bool bShowProgress = true);
+
+	/**
+	 * Remove an actor from available sources.
+	 */
+	UFUNCTION(BlueprintCallable, Category = TakeRecorder)
+	void RemoveActorFromSources(AActor* InActor);
+
+	/** Retrieve the actor from a source, if applicable. */
+	UFUNCTION(BlueprintCallable, Category = TakeRecorder)
+	AActor* GetSourceActor(UTakeRecorderSource* InSource) const;
+	// ~End Sources
+	
+	/** Retrieve the current take recorder state. */
+	UFUNCTION(BlueprintCallable, Category = TakeRecorder)
+	ETakeRecorderState GetState() const;
+
+	/** Directly set the take number. */
+	UFUNCTION(BlueprintCallable, Category = TakeRecorder)
+	void SetTakeNumber(int32 InNewTakeNumber);
+
+	/** Compute the next take number given a slate. */
+	UFUNCTION(BlueprintCallable, Category = TakeRecorder)
+	int32 GetNextTakeNumber(const FString& InSlate) const;
+
+	/** Find both the current maximum take value and the total number of takes for a given slate. */
+	UFUNCTION(BlueprintCallable, Category = TakeRecorder)
+	void GetNumberOfTakes(const FString& InSlate, int32& OutMaxTake, int32& OutNumTakes) const;
+
+	/**
+	 * Retrieve all slates.
+	 * @param InPackagePath [Optional] Scope the search to a specific folder, recursively.
+	 */
+	UFUNCTION(BlueprintCallable, Category = TakeRecorder)
+	TArray<FAssetData> GetSlates(FName InPackagePath = NAME_None) const;
+	
+	/** Directly set the slate name. */
+	UFUNCTION(BlueprintCallable, Category = TakeRecorder)
+	void SetSlateName(const FString& InSlateName);
+
+	/** Mark the current frame. */
+	UFUNCTION(BlueprintCallable, Category = TakeRecorder)
+	bool MarkFrame();
+
+	/**
+	 * Access the frame rate for this take
+	 */
+	UFUNCTION(BlueprintCallable, Category = TakeRecorder)
+	FFrameRate GetFrameRate() const;
+	
+	/**
+	 * Set the frame rate for this take
+	 */
+	UFUNCTION(BlueprintCallable, Category = TakeRecorder)
+	void SetFrameRate(FFrameRate InFrameRate);
+	
+	/**
+	 * Set if the frame rate is set from the Timecode frame rate
+	 */
+	UFUNCTION(BlueprintCallable, Category = TakeRecorder)
+	void SetFrameRateFromTimecode();
+
+	/** Import a preset to the transient preset. */
+	void ImportPreset(const FAssetData& InPreset);
+
+	/** If Take Recorder is currently reviewing. */
+	UFUNCTION(BlueprintCallable, Category = TakeRecorder)
+	bool IsReviewing() const;
+
+	/** If Take Recorder is currently recording. */
+	UFUNCTION(BlueprintCallable, Category = TakeRecorder)
+	bool IsRecording() const;
+
+	/**
+	 * Retrieve the current sequence's countdown.
+	 * @param OutValue The current value of the countdown.
+	 * @return True if we are in a countdown sequence, false otherwise.
+	 */
+	UFUNCTION(BlueprintCallable, Category = TakeRecorder)
+	bool TryGetSequenceCountdown(float& OutValue) const;
+
+	/**
+	 * Sets the current sequence's countdown.
+	 * @param InSeconds Time in seconds.
+	 */
+	UFUNCTION(BlueprintCallable, Category = TakeRecorder)
+	void SetSequenceCountdown(float InSeconds);
+
+	/** Retrieve additional settings objects from a source. */
+	UFUNCTION(BlueprintCallable, Category = TakeRecorder)
+	TArray<UObject*> GetSourceRecordSettings(UTakeRecorderSource* InSource) const;
+
+	/** Retrieve the global take recorder settings. */
+	UFUNCTION(BlueprintCallable, Category = TakeRecorder)
+	FTakeRecorderParameters GetGlobalRecordSettings() const;
+
+	/** Set the global take recorder settings. */
+	UFUNCTION(BlueprintCallable, Category = TakeRecorder)
+	void SetGlobalRecordSettings(const FTakeRecorderParameters& InParameters);
+	
+	/** Retrieve the current meta data. */
+	UFUNCTION(BlueprintCallable, Category = TakeRecorder)
+	UTakeMetaData* GetTakeMetaData() const;
+
+	/** Return the level sequence we are using. */
+	UFUNCTION(BlueprintCallable, Category = TakeRecorder)
+	ULevelSequence* GetLevelSequence() const;
+
+	/** Retrieve the Supplied Level Sequence if it exists. */
+	ULevelSequence* GetSuppliedLevelSequence() const;
+
+	/** Retrieve the Recording Level Sequence if it exists. */
+	ULevelSequence* GetRecordingLevelSequence() const;
+
+	/** Retrieve the Record Into Level Sequence if it exists. */
+	ULevelSequence* GetRecordIntoLevelSequence() const;
+
+	/** Retrieve the Last Recorded Level Sequence if it exists. */
+	ULevelSequence* GetLastRecordedLevelSequence() const;
+
+	/** Retrieve the transient preset the subsystem is using. */
+	UTakePreset* GetTransientPreset() const;
+	
+	/** The current take recorder mode. */
+	UFUNCTION(BlueprintCallable, Category = TakeRecorder)
+	ETakeRecorderMode GetTakeRecorderMode() const;
+
+	/** Retrieve relevant Naming Tokens data for Take Recorder. */
+	UTakeRecorderNamingTokensData* GetNamingTokensData() const;
+	
+	/**
+	 * Retrieve a multi-cast delegate that is triggered when a recording pre-initializes
+	 */
+	FOnTakeRecordingInitialized& GetOnRecordingPreInitializedEvent();
+	/**
+	 * Retrieve a multi-cast delegate that is triggered when a recording initializes
+	 */
+	FOnTakeRecordingInitialized& GetOnRecordingInitializedEvent();
+	/**
+	 * Retrieve a multi-cast delegate that is triggered when a recording starts
+	 */
+	FOnTakeRecordingStarted& GetOnRecordingStartedEvent();
+	/**
+	 * Retrieve a multi-cast delegate that is triggered when a recording is stopped
+	 */
+	FOnTakeRecordingStopped& GetOnRecordingStoppedEvent();
+	/**
+	 * Retrieve a multi-cast delegate that is triggered when a recording finishes
+	 */
+	FOnTakeRecordingFinished& GetOnRecordingFinishedEvent();
+	/**
+	 * Retrieve a multi-cast delegate that is triggered when a recording is cancelled.
+	 */
+	FOnTakeRecordingCancelled& GetOnRecordingCancelledEvent();
+	/**
+	 * Retrieve a multi-cast delegate that is triggered when a source is added.
+	 */
+	UTakeRecorderSources::FOnSourceAdded& GetOnRecordingSourceAddedEvent();
+	/**
+	 * Retrieve a multi-cast delegate that is triggered when a source is removed.
+	 */
+	UTakeRecorderSources::FOnSourceRemoved& GetOnRecordingSourceRemovedEvent();
+	
+private:
+	/**
+	 * Allocate the preset required for interacting with this subsystem. Re-uses an existing preset if necessary
+	 */
+	static UTakePreset* AllocateTransientPreset();
+
+	/** Cache the current available metadata. */
+	void CacheMetaData();
+
+	/** Sets the frame rate. */
+	void SetFrameRateImpl(const FFrameRate& InFrameRate, bool bFromTimecode);
+
+	/** Calculate and apply the next take number. */
+	void IncrementTakeNumber();
+
+	void OnAssetRegistryFilesLoaded();
+	void OnRecordingInitialized(UTakeRecorder* Recorder);
+	void OnRecordingStarted(UTakeRecorder* Recorder);
+	void OnRecordingStopped(UTakeRecorder* Recorder);
+	void OnRecordingFinished(UTakeRecorder* Recorder);
+	void OnRecordingCancelled(UTakeRecorder* Recorder);
+
+	/** Update the last level sequence. */
+	void SetLastLevelSequence(ULevelSequence *InSequence);
+
+	/** Callback when any source is added. */
+	void OnSourceAdded(UTakeRecorderSource* Source);
+
+	/** Callback when any source is removed. */
+	void OnSourceRemoved(UTakeRecorderSource* Source);
+
+	/** Setup bindings to our naming tokens. */
+	void BindToNamingTokenEvents();
+	
+	/** Remove bindings to our naming tokens. */
+	void UnbindNamingTokensEvents();
+
+	/** Called before our naming tokens class evaluates. Used to populate naming token data with our custom token definitions. */
+	void OnTakeRecorderNamingTokensPreEvaluate(const FNamingTokensEvaluationData& InEvaluationData);
+	
+private:
+	/** Last data used for initialization. */
+	FTakeRecorderSequenceParameters TargetSequenceData;
+
+	/** A transient preset available for the subsystem. */
+	UPROPERTY(Transient)
+	TObjectPtr<UTakePreset> TransientPreset;
+
+	/** Current supplied level sequence. */
+	UPROPERTY()
+	TObjectPtr<ULevelSequence> SuppliedLevelSequence;
+
+	/** Current record into level sequence. */
+	UPROPERTY()
+	TObjectPtr<ULevelSequence> RecordIntoLevelSequence;
+
+	/** Current recording level sequence. */
+	UPROPERTY()
+	TObjectPtr<ULevelSequence> RecordingLevelSequence;
+
+	/** The last recorded level sequence. */
+	UPROPERTY()
+	TObjectPtr<ULevelSequence> LastRecordedLevelSequence;
+
+	/** Take meta-data cached from the level sequence if it exists. */
+	UPROPERTY()
+	TObjectPtr<UTakeMetaData> TakeMetaData;
+	
+	/** Transient take meta-data owned by this subsystem. Only used if none exists on the level sequence already. */
+	UPROPERTY(Transient)
+	TObjectPtr<UTakeMetaData> TransientTakeMetaData;
+
+	/** Stored data relevant to our naming tokens. Managed as UObject for transactions. */
+	UPROPERTY(Transient)
+	TObjectPtr<UTakeRecorderNamingTokensData> NamingTokensData;
+	
+	// These native delegates are declared elsewhere throughout Take Recorder and we reuse them here. However, it's important
+	// the editor subsystem manage its own variants because the order they fire is important for listeners, such as the panel. For example, the
+	// editor subsystem may store a strong reference to a sequence which needs to happen before the panel run its logic, otherwise the panel
+	// may end up clearing the only strong reference prematurely. Therefor the panel can't directly listen to UTakeRecorder::OnRecordingStartedEvent,
+	// but can listen to the editor subsystem's version.
+	
+	/** Multicast delegate when recording is pre-initialized. For native binding. */
+	FOnTakeRecordingPreInitialize OnRecordingPreInitializeEvent;
+	/** Multicast delegate when recording is initialized. For native binding. */
+	FOnTakeRecordingInitialized ORecordingInitializedEvent;
+	/** Multicast delegate for when a recording is started. For native binding. */
+	FOnTakeRecordingStarted OnRecordingStartedEvent;
+	/** Multicast delegate for when a recording has stopped. For native binding. */
+	FOnTakeRecordingStopped OnRecordingStoppedEvent;
+	/** Multicast delegate for when a recording has finished. For native binding. */
+	FOnTakeRecordingFinished OnRecordingFinishedEvent;
+	/** Multicast delegate for when a recording has been cancelled. For native binding. */
+	FOnTakeRecordingCancelled OnRecordingCancelledEvent;
+	/** Multicast delegate when any source has been added. For native binding. */
+	UTakeRecorderSources::FOnSourceAdded OnRecordingSourceAddedEvent;
+	/** Multicast delegate when any source has been added. For native binding. */
+	UTakeRecorderSources::FOnSourceRemoved OnRecordingSourceRemovedEvent;
+	
+	FDelegateHandle OnAssetRegistryFilesLoadedHandle;
+	FDelegateHandle OnRecordingInitializedHandle, OnRecordingStoppedHandle, OnRecordingFinishedHandle, OnRecordingCancelledHandle;
+	
+	/** Whether we should auto apply the next available take number when asset discovery has finished or not */
+	bool bAutoApplyTakeNumber;
+	
 public:
 
 	/** Called before initialization occurs (ie. when the recording button is pressed and before the countdown starts) */
-	UPROPERTY(BlueprintAssignable, Category = "Take Recorder")
+	UPROPERTY(BlueprintAssignable, Category=TakeRecorder)
 	FTakeRecorderPreInitialize TakeRecorderPreInitialize;
 
+	/** called when take recorder is initializing. */
+	UPROPERTY(BlueprintAssignable, Category=TakeRecorder)
+	FTakeRecorderInitialized TakeRecorderInitialized;
+	
 	/** Called when take recorder is started */
-	UPROPERTY(BlueprintAssignable, Category = "Take Recorder")
+	UPROPERTY(BlueprintAssignable, Category=TakeRecorder)
 	FTakeRecorderStarted TakeRecorderStarted;
 
 	/** Called when take recorder is stopped */
-	UPROPERTY(BlueprintAssignable, Category = "Take Recorder")
+	UPROPERTY(BlueprintAssignable, Category=TakeRecorder)
 	FTakeRecorderStopped TakeRecorderStopped;
 
 	/** Called when take recorder has finished */
-	UPROPERTY(BlueprintAssignable, Category = "Take Recorder")
+	UPROPERTY(BlueprintAssignable, Category=TakeRecorder)
 	FTakeRecorderFinished TakeRecorderFinished;
 
 	/** Called when take recorder is cancelled */
-	UPROPERTY(BlueprintAssignable, Category = "Take Recorder")
+	UPROPERTY(BlueprintAssignable, Category=TakeRecorder)
 	FTakeRecorderCancelled TakeRecorderCancelled;
 
 	/** Called when a marked frame is added to take recorder */
-	UPROPERTY(BlueprintAssignable, Category = "Take Recorder")
+	UPROPERTY(BlueprintAssignable, Category=TakeRecorder)
 	FTakeRecorderMarkedFrameAdded TakeRecorderMarkedFrameAdded;
-};
\ No newline at end of file
+
+	/** Multicast delegate when any source has been added. */
+	UPROPERTY(BlueprintAssignable, Category=TakeRecorder)
+	FOnTakeRecorderSourceAddedDynamic TakeRecorderSourceAdded;
+
+	/** Multicast delegate when any source has been removed. */
+	UPROPERTY(BlueprintAssignable, Category=TakeRecorder)
+	FOnTakeRecorderSourceRemovedDynamic TakeRecorderSourceRemoved;
+};
diff --git a/Engine/Plugins/VirtualProduction/Takes/Source/TakeRecorderSources/Private/TakeRecorderSourceHelpers.cpp b/Engine/Plugins/VirtualProduction/Takes/Source/TakeRecorderSources/Private/TakeRecorderSourceHelpers.cpp
--- a/Engine/Plugins/VirtualProduction/Takes/Source/TakeRecorderSources/Private/TakeRecorderSourceHelpers.cpp
+++ b/Engine/Plugins/VirtualProduction/Takes/Source/TakeRecorderSources/Private/TakeRecorderSourceHelpers.cpp
@@ -35,6 +35,10 @@ void AddActorSources(
 
 		for (AActor* Actor : ActorsToRecord)
 		{
+			if (!Actor)
+			{
+				continue;
+			}
 			if (Actor->IsA<ALevelSequenceActor>())
 			{
 				ALevelSequenceActor* LevelSequenceActor = Cast<ALevelSequenceActor>(Actor);
@@ -87,6 +91,30 @@ void AddActorSources(
 	}
 }
 
+void RemoveActorSources(UTakeRecorderSources* TakeRecorderSources, TArrayView<AActor* const> ActorsToRemove)
+{
+	if (TakeRecorderSources->GetSources().IsEmpty())
+	{
+		return;
+	}
+
+	FScopedTransaction Transaction(FText::Format(LOCTEXT(
+			"RemoveActorSources", 
+			"Remove Recording {0}|plural(one=Source, other=Sources)"), ActorsToRemove.Num()));
+	TakeRecorderSources->Modify();
+
+	for (UTakeRecorderSource* Source : TakeRecorderSources->GetSourcesCopy())
+	{
+		if (const UTakeRecorderActorSource* ActorSource = Cast<UTakeRecorderActorSource>(Source))
+		{
+			if (ActorsToRemove.Contains(ActorSource->Target))
+			{
+				TakeRecorderSources->RemoveSource(Source);
+			}
+		}
+	}
+}
+
 void RemoveAllActorSources(UTakeRecorderSources* TakeRecorderSources)
 {
 	if (TakeRecorderSources->GetSources().IsEmpty())
@@ -104,6 +132,16 @@ void RemoveAllActorSources(UTakeRecorderSources* TakeRecorderSources)
 		TakeRecorderSources->RemoveSource(TakeRecorderSources->GetSources()[0]);
 	}
 }
+	
+AActor* GetSourceActor(UTakeRecorderSource* Source)
+{
+	if (const UTakeRecorderActorSource* ActorSource = Cast<UTakeRecorderActorSource>(Source))
+	{
+		return ActorSource->Target.LoadSynchronous();
+	}
+
+	return nullptr;
+}
 
 void ProcessRecordedTimes(ULevelSequence* InSequence, UMovieSceneTakeTrack* TakeTrack, const TOptional<TRange<FFrameNumber>>& FrameRange, const FArrayOfRecordedTimePairs& RecordedTimes)
 {
diff --git a/Engine/Plugins/VirtualProduction/Takes/Source/TakeRecorderSources/Private/TakeRecorderSourcesModule.cpp b/Engine/Plugins/VirtualProduction/Takes/Source/TakeRecorderSources/Private/TakeRecorderSourcesModule.cpp
--- a/Engine/Plugins/VirtualProduction/Takes/Source/TakeRecorderSources/Private/TakeRecorderSourcesModule.cpp
+++ b/Engine/Plugins/VirtualProduction/Takes/Source/TakeRecorderSources/Private/TakeRecorderSourcesModule.cpp
@@ -152,6 +152,12 @@ void FTakeRecorderSourcesModule::StartupModule()
 
 	SourcesMenuExtension = TakeRecorderModule.RegisterSourcesMenuExtension(FOnExtendSourcesMenu::CreateStatic(ExtendSourcesMenu));
 
+	ITakeRecorderModule::FSourceExtensionData ExtensionData;
+	ExtensionData.OnAddActorSource = ITakeRecorderModule::FSourceExtensionData::FAddActorSource::CreateStatic(TakeRecorderSourceHelpers::AddActorSources);
+	ExtensionData.OnRemoveActorSource = ITakeRecorderModule::FSourceExtensionData::FRemoveActorSource::CreateStatic(TakeRecorderSourceHelpers::RemoveActorSources);
+	ExtensionData.OnGetSourceActor = ITakeRecorderModule::FSourceExtensionData::FGetSourceActor::CreateStatic(TakeRecorderSourceHelpers::GetSourceActor);
+	TakeRecorderModule.RegisterSourcesExtension(ExtensionData);
+	
 	TakeRecorderModule.RegisterSettingsObject(GetMutableDefault<UTakeRecorderMicrophoneAudioSourceSettings>());
 	TakeRecorderModule.RegisterSettingsObject(GetMutableDefault<UTakeRecorderMicrophoneAudioManager>());
 	TakeRecorderModule.RegisterSettingsObject(GetMutableDefault<UMovieSceneAnimationTrackRecorderEditorSettings>());
@@ -175,6 +181,7 @@ void FTakeRecorderSourcesModule::ShutdownModule()
 	if (TakeRecorderModule)
 	{
 		TakeRecorderModule->UnregisterSourcesMenuExtension(SourcesMenuExtension);
+		TakeRecorderModule->UnregisterSourcesExtension();
 	}
 
 	FTakeRecorderSourcesCommands::Unregister();
diff --git a/Engine/Plugins/VirtualProduction/Takes/Source/TakeRecorderSources/Public/TakeRecorderSourceHelpers.h b/Engine/Plugins/VirtualProduction/Takes/Source/TakeRecorderSources/Public/TakeRecorderSourceHelpers.h
--- a/Engine/Plugins/VirtualProduction/Takes/Source/TakeRecorderSources/Public/TakeRecorderSourceHelpers.h
+++ b/Engine/Plugins/VirtualProduction/Takes/Source/TakeRecorderSources/Public/TakeRecorderSourceHelpers.h
@@ -27,11 +27,25 @@ namespace TakeRecorderSourceHelpers
 		bool                      bReduceKeys = true,
 		bool                      bShowProgress = true);
 
+	/**
+	 * Remove specific actors as a source.
+	 * @param TakeRecorderSources The list of sources used for the current take. 
+	 * @param ActorsToRemove The list of Actors to be removed from Sources.
+	 */
+	TAKERECORDERSOURCES_API void RemoveActorSources(UTakeRecorderSources* TakeRecorderSources, TArrayView<AActor* const> ActorsToRemove);
+	
 	/**
 	 * Removes all sources from a list of sources to record from.
 	 */
 	TAKERECORDERSOURCES_API void RemoveAllActorSources(UTakeRecorderSources* Sources);
 
+	/**
+	 * Retrieve the source actor if applicable.
+	 * @param Source The source to check.
+	 * @return The source actor, if one is set.
+	 */
+	TAKERECORDERSOURCES_API AActor* GetSourceActor(UTakeRecorderSource* Source);
+
 	using FArrayOfRecordedTimePairs = const TArray<TPair<FQualifiedFrameTime, FQualifiedFrameTime>>;
 
 	/**
diff --git a/Engine/Plugins/VirtualProduction/Takes/Source/TakesCore/Private/TakeRecorderSources.cpp b/Engine/Plugins/VirtualProduction/Takes/Source/TakesCore/Private/TakeRecorderSources.cpp
--- a/Engine/Plugins/VirtualProduction/Takes/Source/TakesCore/Private/TakeRecorderSources.cpp
+++ b/Engine/Plugins/VirtualProduction/Takes/Source/TakesCore/Private/TakeRecorderSources.cpp
@@ -51,6 +51,8 @@ UTakeRecorderSource* UTakeRecorderSources::AddSource(TSubclassOf<UTakeRecorderSo
 			NewSource->Initialize();
 			Sources.Add(NewSource);
 			++SourcesSerialNumber;
+
+			OnSourceAdded().Broadcast(NewSource);
 		}
 	}
 
@@ -65,6 +67,8 @@ void UTakeRecorderSources::RemoveSource(UTakeRecorderSource* InSource)
 	SourceSubSequenceMap.Remove(InSource);
 
 	++SourcesSerialNumber;
+
+	OnSourceRemoved().Broadcast(InSource);
 }
 
 FDelegateHandle UTakeRecorderSources::BindSourcesChanged(const FSimpleDelegate& Handler)
@@ -77,6 +81,18 @@ void UTakeRecorderSources::UnbindSourcesChanged(FDelegateHandle Handle)
 	OnSourcesChangedEvent.Remove(Handle);
 }
 
+UTakeRecorderSources::FOnSourceAdded& UTakeRecorderSources::OnSourceAdded()
+{
+	static FOnSourceAdded OnSourceAddedEvent;
+	return OnSourceAddedEvent;
+}
+
+UTakeRecorderSources::FOnSourceRemoved& UTakeRecorderSources::OnSourceRemoved()
+{
+	static FOnSourceRemoved OnSourceRemovedEvent;
+	return OnSourceRemovedEvent;
+}
+
 void UTakeRecorderSources::PostEditChangeProperty(FPropertyChangedEvent& PropertyChangedEvent)
 {
 	Super::PostEditChangeProperty(PropertyChangedEvent);
diff --git a/Engine/Plugins/VirtualProduction/Takes/Source/TakesCore/Public/TakeRecorderSources.h b/Engine/Plugins/VirtualProduction/Takes/Source/TakesCore/Public/TakeRecorderSources.h
--- a/Engine/Plugins/VirtualProduction/Takes/Source/TakesCore/Public/TakeRecorderSources.h
+++ b/Engine/Plugins/VirtualProduction/Takes/Source/TakesCore/Public/TakeRecorderSources.h
@@ -119,7 +119,16 @@ public:
 	 * @param Handle          The handle returned from BindSourcesChanged for the handler to remove
 	 */
 	void UnbindSourcesChanged(FDelegateHandle Handle);
+	
+	DECLARE_MULTICAST_DELEGATE_OneParam(FOnSourceAdded, UTakeRecorderSource* Source);
+	DECLARE_MULTICAST_DELEGATE_OneParam(FOnSourceRemoved, UTakeRecorderSource* Source);
+	
+	/** Multicast delegate for when any source is added. */
+	static FOnSourceAdded& OnSourceAdded();
 
+	/** Multicast delegate for when any source is removed. */
+	static FOnSourceRemoved& OnSourceRemoved();
+	
 public:
 
 	/*
