TakeRecorder: Fix static analysis warnings.

diff --git a/Engine/Plugins/VirtualProduction/Takes/Source/TakeRecorder/Private/Widgets/STakeRecorderCockpit.cpp b/Engine/Plugins/VirtualProduction/Takes/Source/TakeRecorder/Private/Widgets/STakeRecorderCockpit.cpp
--- a/Engine/Plugins/VirtualProduction/Takes/Source/TakeRecorder/Private/Widgets/STakeRecorderCockpit.cpp
+++ b/Engine/Plugins/VirtualProduction/Takes/Source/TakeRecorder/Private/Widgets/STakeRecorderCockpit.cpp
@@ -161,6 +161,7 @@ void STakeRecorderCockpit::Construct(const FArguments& InArgs)
 	CachedTakeNumber = -1;
 
 	TakeRecorderSubsystem = GEngine->GetEngineSubsystem<UTakeRecorderSubsystem>();
+	check(TakeRecorderSubsystem.IsValid());
 	
 	UpdateTakeError();
 	UpdateRecordError();
@@ -538,12 +539,18 @@ void STakeRecorderCockpit::UpdateRecordError()
 		RecordErrorText = LOCTEXT("ErrorWidget_NoSlate", "You must enter a slate to begin recording.");
 		return;
 	}
-	FString PackageName = TakeMetaData->GenerateAssetPath(GetDefault<UTakeRecorderProjectSettings>()->Settings.GetTakeAssetPath());
-	FText OutReason;
-	if (!FPackageName::IsValidLongPackageName(PackageName, false, &OutReason))
+
+	FString PackageName;
+
+	if (TakeMetaData)
 	{
-		RecordErrorText = FText::Format(LOCTEXT("ErrorWidget_InvalidPath", "{0} is not a valid asset path. {1}"), FText::FromString(PackageName), OutReason);
-		return;
+		PackageName = TakeMetaData->GenerateAssetPath(GetDefault<UTakeRecorderProjectSettings>()->Settings.GetTakeAssetPath());
+		FText OutReason;
+		if (!FPackageName::IsValidLongPackageName(PackageName, false, &OutReason))
+		{
+			RecordErrorText = FText::Format(LOCTEXT("ErrorWidget_InvalidPath", "{0} is not a valid asset path. {1}"), FText::FromString(PackageName), OutReason);
+			return;
+		}
 	}
 
 	if (TakeMetaData && TakeMetaData->GetFrameRateFromTimecode())
