TakeRecorder: Add null checks in case the subsystem is being accessed before TakeMetaData is available.

diff --git a/Engine/Plugins/VirtualProduction/Takes/Source/TakeRecorder/Private/Recorder/TakeRecorder.cpp b/Engine/Plugins/VirtualProduction/Takes/Source/TakeRecorder/Private/Recorder/TakeRecorder.cpp
--- a/Engine/Plugins/VirtualProduction/Takes/Source/TakeRecorder/Private/Recorder/TakeRecorder.cpp
+++ b/Engine/Plugins/VirtualProduction/Takes/Source/TakeRecorder/Private/Recorder/TakeRecorder.cpp
@@ -454,6 +454,15 @@ bool UTakeRecorder::Initialize ( ULevelSequence* LevelSequenceBase, UTakeRecorde
 		return false;
 	}
 
+	if (MetaData == nullptr)
+	{
+		if (OutError)
+		{
+			*OutError = LOCTEXT("NoMetaDataError", "No Meta Data specified.");
+		}
+		return false;
+	}
+
 	if (MetaData->GetSlate().IsEmpty())
 	{
 		if (OutError)
diff --git a/Engine/Plugins/VirtualProduction/Takes/Source/TakeRecorder/Private/Recorder/TakeRecorderSubsystem.cpp b/Engine/Plugins/VirtualProduction/Takes/Source/TakeRecorder/Private/Recorder/TakeRecorderSubsystem.cpp
--- a/Engine/Plugins/VirtualProduction/Takes/Source/TakeRecorder/Private/Recorder/TakeRecorderSubsystem.cpp
+++ b/Engine/Plugins/VirtualProduction/Takes/Source/TakeRecorder/Private/Recorder/TakeRecorderSubsystem.cpp
@@ -497,10 +497,13 @@ ETakeRecorderState UTakeRecorderSubsystem::GetState() const
 
 void UTakeRecorderSubsystem::SetTakeNumber(int32 InNewTakeNumber)
 {
-	FScopedTransaction Transaction(LOCTEXT("SetTake_Transaction", "Set Take Number"));
-	TakeMetaData->Modify();
-	TakeMetaData->SetTakeNumber(InNewTakeNumber, false);
-	bAutoApplyTakeNumber = false;
+	if (TakeMetaData)
+	{
+		FScopedTransaction Transaction(LOCTEXT("SetTake_Transaction", "Set Take Number"));
+		TakeMetaData->Modify();
+		TakeMetaData->SetTakeNumber(InNewTakeNumber, false);
+		bAutoApplyTakeNumber = false;
+	}
 }
 
 int32 UTakeRecorderSubsystem::GetNextTakeNumber(const FString& InSlate) const
@@ -550,7 +553,7 @@ TArray<FAssetData> UTakeRecorderSubsystem::GetSlates(FName InPackagePath) const
 
 void UTakeRecorderSubsystem::SetSlateName(const FString& InSlateName)
 {
-	if (TakeMetaData->GetSlate() != InSlateName)
+	if (TakeMetaData && TakeMetaData->GetSlate() != InSlateName)
 	{
 		FScopedTransaction Transaction(LOCTEXT("SetSlate_Transaction", "Set Take Slate"));
 		TakeMetaData->Modify();
@@ -568,7 +571,7 @@ void UTakeRecorderSubsystem::SetSlateName(const FString& InSlateName)
 
 bool UTakeRecorderSubsystem::MarkFrame()
 {
-	if (UTakeRecorderBlueprintLibrary::IsRecording())
+	if (UTakeRecorderBlueprintLibrary::IsRecording() && TakeMetaData)
 	{
 		const FFrameRate FrameRate = TakeMetaData->GetFrameRate();
 
