Take Recorder: Allow users to specify subclass of ULevelSequence

diff --git a/Engine/Plugins/VirtualProduction/Takes/Source/TakeRecorder/Private/Customization/TakePresetRecorderCustomization.cpp b/Engine/Plugins/VirtualProduction/Takes/Source/TakeRecorder/Private/Customization/TakePresetRecorderCustomization.cpp
new file mode 100644
--- /dev/null
+++ b/Engine/Plugins/VirtualProduction/Takes/Source/TakeRecorder/Private/Customization/TakePresetRecorderCustomization.cpp
@@ -0,0 +1,81 @@
+﻿// Copyright Epic Games, Inc. All Rights Reserved.
+
+#include "TakePresetRecorderCustomization.h"
+
+#include "ClassViewerFilter.h"
+#include "DetailLayoutBuilder.h"
+#include "DetailWidgetRow.h"
+#include "LevelSequence.h"
+#include "PropertyCustomizationHelpers.h"
+#include "TakePresetSettings.h"
+
+#define LOCTEXT_NAMESPACE "FTakePresetRecorderCustomization"
+
+namespace UE::TakeRecorder
+{
+class FAssetClassParentFilter : public IClassViewerFilter
+{
+public:
+	/** All children of these classes will be included unless filtered out by another setting. */
+	TSet< const UClass* > AllowedChildrenOfClasses;
+
+	/** Disallowed class flags. */
+	EClassFlags DisallowedClassFlags;
+
+	virtual bool IsClassAllowed(const FClassViewerInitializationOptions& InInitOptions, const UClass* InClass, TSharedRef< FClassViewerFilterFuncs > InFilterFuncs) override
+	{
+		return !InClass->HasAnyClassFlags(DisallowedClassFlags)
+			&& InFilterFuncs->IfInChildOfClassesSet(AllowedChildrenOfClasses, InClass) != EFilterReturn::Failed;
+	}
+
+	virtual bool IsUnloadedClassAllowed(const FClassViewerInitializationOptions& InInitOptions, const TSharedRef< const IUnloadedBlueprintData > InUnloadedClassData, TSharedRef< FClassViewerFilterFuncs > InFilterFuncs) override
+	{
+		return !InUnloadedClassData->HasAnyClassFlags(DisallowedClassFlags)
+			&& InFilterFuncs->IfInChildOfClassesSet(AllowedChildrenOfClasses, InUnloadedClassData) != EFilterReturn::Failed;
+	}
+};
+
+void FTakePresetRecorderCustomization::CustomizeHeader(
+	TSharedRef<IPropertyHandle> PropertyHandle, FDetailWidgetRow& HeaderRow, IPropertyTypeCustomizationUtils& CustomizationUtils
+	)
+{
+	const TSharedPtr<IPropertyHandle> TargetRecordClass = PropertyHandle->GetChildHandle(GET_MEMBER_NAME_CHECKED(FTakeRecorderTargetRecordClassProperty, TargetRecordClass));
+	if (!ensure(TargetRecordClass))
+	{
+		return;
+	}
+	
+	TSharedRef<FAssetClassParentFilter> ClassFilter = MakeShared<FAssetClassParentFilter>();
+	ClassFilter->DisallowedClassFlags = CLASS_Abstract | CLASS_Deprecated | CLASS_NewerVersionExists;
+	ClassFilter->AllowedChildrenOfClasses.Add(ULevelSequence::StaticClass());
+
+	// The point of this is to customize the TargetRecordClass and only trigger a PostEditChange if the user clicked "Ok" in the dialogue.
+	// A dialogue will only be shown if the user has made changes that would be discarded. If there's none, the property is silently changed.
+	HeaderRow
+		.NameContent() [ TargetRecordClass->CreatePropertyNameWidget() ]
+		.ValueContent()
+		[
+			SNew(SClassPropertyEntryBox)
+			.MetaClass(UObject::StaticClass())
+			.AllowNone(false)
+			.AllowAbstract(false)
+			.SelectedClass_Lambda([this, TargetRecordClass]
+			{
+				UObject* ClassObject = nullptr;
+				TargetRecordClass->GetValue(ClassObject);
+				return Cast<UClass>(ClassObject);
+				
+			})
+			.ClassViewerFilters({ ClassFilter })
+			.OnSetClass_Lambda([this, TargetRecordClass](const UClass* Class)
+			{
+				if (Class && PromptChangeTargetRecordClassDelegate.Execute(Class))
+				{
+					TargetRecordClass->SetValue(Class);
+				}
+			})
+		];
+}
+}
+
+#undef LOCTEXT_NAMESPACE
\ No newline at end of file
diff --git a/Engine/Plugins/VirtualProduction/Takes/Source/TakeRecorder/Private/Customization/TakePresetRecorderCustomization.h b/Engine/Plugins/VirtualProduction/Takes/Source/TakeRecorder/Private/Customization/TakePresetRecorderCustomization.h
new file mode 100644
--- /dev/null
+++ b/Engine/Plugins/VirtualProduction/Takes/Source/TakeRecorder/Private/Customization/TakePresetRecorderCustomization.h
@@ -0,0 +1,33 @@
+﻿// Copyright Epic Games, Inc. All Rights Reserved.
+
+#pragma once
+
+#include "Delegates/Delegate.h"
+#include "IPropertyTypeCustomization.h"
+
+namespace UE::TakeRecorder
+{
+/** Possibly asks the user whether they are sure they want to change TargetRecordClass. Returns whether the change can be made. */
+DECLARE_DELEGATE_RetVal_OneParam(bool, FPromptChangeTargetRecordClass, const UClass* /*NewClass*/);
+	
+/** Asks the user whether they're sure they want to change the UTakePresetSettings::TargetRecordClass, as that usually requires clearing the pending change. */
+class FTakePresetRecorderCustomization : public IPropertyTypeCustomization
+{
+public:
+
+	explicit FTakePresetRecorderCustomization(FPromptChangeTargetRecordClass PromptUserDelegate)
+		: PromptChangeTargetRecordClassDelegate(MoveTemp(PromptUserDelegate))
+	{}
+
+	//~ Begin IDetailCustomization Interface
+	virtual void CustomizeHeader(TSharedRef<IPropertyHandle> PropertyHandle, FDetailWidgetRow& HeaderRow, IPropertyTypeCustomizationUtils& CustomizationUtils) override;
+	virtual void CustomizeChildren(TSharedRef<IPropertyHandle> PropertyHandle, IDetailChildrenBuilder& ChildBuilder, IPropertyTypeCustomizationUtils& CustomizationUtils) override {}
+	//~ End IDetailCustomization Interface
+
+private:
+	
+	/** Possibly asks the user whether they are sure they want to change TargetRecordClass. Returns whether the change can be made. */
+	FPromptChangeTargetRecordClass PromptChangeTargetRecordClassDelegate;
+};
+
+}
diff --git a/Engine/Plugins/VirtualProduction/Takes/Source/TakeRecorder/Private/TakeRecorderModule.cpp b/Engine/Plugins/VirtualProduction/Takes/Source/TakeRecorder/Private/TakeRecorderModule.cpp
--- a/Engine/Plugins/VirtualProduction/Takes/Source/TakeRecorder/Private/TakeRecorderModule.cpp
+++ b/Engine/Plugins/VirtualProduction/Takes/Source/TakeRecorder/Private/TakeRecorderModule.cpp
@@ -24,6 +24,7 @@
 #include "ITakeRecorderDropHandler.h"
 #include "ISettingsModule.h"
 #include "TakeRecorderSettings.h"
+#include "TakePresetSettings.h"
 
 #include "Widgets/STakeRecorderTabContent.h"
 #include "Widgets/STakeRecorderCockpit.h"
@@ -583,6 +584,7 @@ void FTakeRecorderModule::UnregisterLevelEditorExtensions()
 void FTakeRecorderModule::RegisterSettings()
 {
 	RegisterSettingsObject(GetMutableDefault<UMovieSceneTakeSettings>());
+	RegisterSettingsObject(GetMutableDefault<UTakePresetSettings>());
 
 	ISettingsModule& SettingsModule = FModuleManager::LoadModuleChecked<ISettingsModule>("Settings");
 
diff --git a/Engine/Plugins/VirtualProduction/Takes/Source/TakeRecorder/Private/Widgets/SLevelSequenceTakeEditor.cpp b/Engine/Plugins/VirtualProduction/Takes/Source/TakeRecorder/Private/Widgets/SLevelSequenceTakeEditor.cpp
--- a/Engine/Plugins/VirtualProduction/Takes/Source/TakeRecorder/Private/Widgets/SLevelSequenceTakeEditor.cpp
+++ b/Engine/Plugins/VirtualProduction/Takes/Source/TakeRecorder/Private/Widgets/SLevelSequenceTakeEditor.cpp
@@ -3,11 +3,13 @@
 #include "Widgets/SLevelSequenceTakeEditor.h"
 #include "Customization/AudioInputChannelPropertyCustomization.h"
 #include "Customization/TakeRecorderAudioSettingsCustomization.h"
+#include "Customization/TakePresetRecorderCustomization.h"
 #include "Customization/RecordedPropertyCustomization.h"
 #include "Customization/RecorderPropertyMapCustomization.h"
 #include "Customization/RecorderSourceObjectCustomization.h"
 #include "TakePreset.h"
 #include "TakeMetaData.h"
+#include "TakePresetSettings.h"
 #include "TakeRecorderSource.h"
 #include "TakeRecorderSourceProperty.h"
 #include "TakeRecorderSources.h"
@@ -22,6 +24,7 @@
 #include "UObject/UObjectIterator.h"
 #include "Templates/SubclassOf.h"
 #include "ClassIconFinder.h"
+#include "Misc/MessageDialog.h"
 
 // AssetRegistry includes
 #include "AssetRegistry/AssetRegistryModule.h"
@@ -102,6 +105,7 @@ void SLevelSequenceTakeEditor::Construct(const FArguments& InArgs)
 
 	OnDetailsPropertiesChangedEvent = InArgs._OnDetailsPropertiesChanged;
 	OnDetailsViewAddedEvent = InArgs._OnDetailsViewAdded;
+	ShouldPromptUserForTargetRecordClassChangeEvent = InArgs._ShouldPromptUserForTargetRecordClassChange;
 
 	DetailsBox = SNew(SScrollBox);
 	DetailsBox->SetScrollBarRightClickDragAllowed(true);
@@ -282,6 +286,21 @@ void SLevelSequenceTakeEditor::OnDetailsPropertiesChanged(const FPropertyChanged
 	OnDetailsPropertiesChangedEvent.ExecuteIfBound(InEvent);
 }
 
+bool SLevelSequenceTakeEditor::PromptUserForTargetRecordClassChange(const UClass* NewClass) const
+{
+	if (ShouldPromptUserForTargetRecordClassChangeEvent.IsBound()
+		&& ShouldPromptUserForTargetRecordClassChangeEvent.Execute(NewClass) == EShowPrompt::AllowOperationSilently)
+	{
+		return true;
+	}
+	
+	const FText WarningMessage(LOCTEXT(
+		"Warning_ChangeTargetLevelSequenceClass",
+		"Changing the class requires clearing the pending take.\nYour current changes will be discarded.\n\nDo you want to clear the pending take?")
+		);
+	return FMessageDialog::Open(EAppMsgType::OkCancel, WarningMessage) == EAppReturnType::Ok;
+}
+
 void SLevelSequenceTakeEditor::AddDetails(const TPair<const UClass*, TArray<UObject*> >& Pair, TArray<FObjectKey>& PreviousClasses)
 {
 	FPropertyEditorModule& PropertyEditorModule = FModuleManager::Get().LoadModuleChecked<FPropertyEditorModule>("PropertyEditor");
@@ -310,6 +329,13 @@ void SLevelSequenceTakeEditor::AddDetails(const TPair<const UClass*, TArray<UObj
 		// We may want to change this in future but it seems like the neatest way to make top level categories have helpful names.
 		Details->RegisterInstancedCustomPropertyLayout(UTakeRecorderSource::StaticClass(), FOnGetDetailCustomizationInstance::CreateLambda(&MakeShared<FRecorderSourceObjectCustomization>));
 
+		Details->RegisterInstancedCustomPropertyTypeLayout(FTakeRecorderTargetRecordClassProperty::StaticStruct()->GetFName(),
+			FOnGetPropertyTypeCustomizationInstance::CreateLambda([this]
+			{
+				return MakeShared<FTakePresetRecorderCustomization>(
+					FPromptChangeTargetRecordClass::CreateSP(this, &SLevelSequenceTakeEditor::PromptUserForTargetRecordClassChange)
+					);
+			}));
 		Details->RegisterInstancedCustomPropertyTypeLayout(FName(TEXT("AudioInputDeviceChannelProperty")), FOnGetPropertyTypeCustomizationInstance::CreateLambda(&MakeShared<FAudioInputChannelPropertyCustomization >));
 		Details->RegisterInstancedCustomPropertyTypeLayout(FName(TEXT("ActorRecorderPropertyMap")), FOnGetPropertyTypeCustomizationInstance::CreateLambda(&MakeShared<FRecorderPropertyMapCustomization >));
 		Details->RegisterInstancedCustomPropertyTypeLayout(FName(TEXT("ActorRecordedProperty")), FOnGetPropertyTypeCustomizationInstance::CreateLambda(&MakeShared<FRecordedPropertyCustomization >));
diff --git a/Engine/Plugins/VirtualProduction/Takes/Source/TakeRecorder/Private/Widgets/SLevelSequenceTakeEditor.h b/Engine/Plugins/VirtualProduction/Takes/Source/TakeRecorder/Private/Widgets/SLevelSequenceTakeEditor.h
--- a/Engine/Plugins/VirtualProduction/Takes/Source/TakeRecorder/Private/Widgets/SLevelSequenceTakeEditor.h
+++ b/Engine/Plugins/VirtualProduction/Takes/Source/TakeRecorder/Private/Widgets/SLevelSequenceTakeEditor.h
@@ -28,6 +28,9 @@ public:
 	DECLARE_DELEGATE_OneParam(FOnDetailsPropertiesChanged, const FPropertyChangedEvent&);
 	DECLARE_DELEGATE_OneParam(FOnDetailsViewAdded, const TWeakPtr<IDetailsView>&);
 
+	enum class EShowPrompt{ ShowPrompt, AllowOperationSilently };
+	DECLARE_DELEGATE_RetVal_OneParam(EShowPrompt, FShouldPromptUserForTargetRecordClassChange, const UClass* /*NewClass*/)
+
 	SLATE_BEGIN_ARGS(SLevelSequenceTakeEditor)
 		: _LevelSequence(nullptr)
 		{}
@@ -35,6 +38,8 @@ public:
 		SLATE_ATTRIBUTE(ULevelSequence*, LevelSequence)
 		SLATE_EVENT(FOnDetailsPropertiesChanged, OnDetailsPropertiesChanged)
 		SLATE_EVENT(FOnDetailsViewAdded, OnDetailsViewAdded)
+		/** Called when the user attempts to change UTakePresetSettings::TargetRecordClass. Returns whether it is necessary to show a prompt to the user. */
+		SLATE_EVENT(FShouldPromptUserForTargetRecordClassChange, ShouldPromptUserForTargetRecordClassChange)
 
 	SLATE_END_ARGS()
 
@@ -93,6 +98,9 @@ private:
 	 */
 	void OnDetailsPropertiesChanged(const FPropertyChangedEvent& InEvent);
 
+	/** Asks the user whether they want to really change UTakePresetSettings::TargetRecordClass. */
+	bool PromptUserForTargetRecordClassChange(const UClass* NewClass) const;
+
 private:
 
 	bool bRequestDetailsRefresh;
@@ -109,4 +117,6 @@ private:
 	FOnDetailsPropertiesChanged OnDetailsPropertiesChangedEvent;
 	/** Called when a details view has been added. */
 	FOnDetailsViewAdded OnDetailsViewAddedEvent;
+	/** Called when the user attempts to change UTakePresetSettings::TargetRecordClass. Returns whether it is necessary to show a prompt to the user. */
+	FShouldPromptUserForTargetRecordClassChange ShouldPromptUserForTargetRecordClassChangeEvent;
 };
\ No newline at end of file
diff --git a/Engine/Plugins/VirtualProduction/Takes/Source/TakeRecorder/Private/Widgets/STakeRecorderPanel.cpp b/Engine/Plugins/VirtualProduction/Takes/Source/TakeRecorder/Private/Widgets/STakeRecorderPanel.cpp
--- a/Engine/Plugins/VirtualProduction/Takes/Source/TakeRecorder/Private/Widgets/STakeRecorderPanel.cpp
+++ b/Engine/Plugins/VirtualProduction/Takes/Source/TakeRecorder/Private/Widgets/STakeRecorderPanel.cpp
@@ -67,6 +67,10 @@
 
 #include "LevelEditor.h"
 #include "Recorder/TakeRecorderSubsystem.h"
+#include "TakePresetSettings.h"
+#include "Editor/Transactor.h"
+#include "Editor/TransBuffer.h"
+#include "Misc/ITransaction.h"
 
 #define LOCTEXT_NAMESPACE "STakeRecorderPanel"
 
@@ -87,6 +91,8 @@ STakeRecorderPanel::~STakeRecorderPanel()
 UE_DISABLE_OPTIMIZATION_SHIP
 void STakeRecorderPanel::Construct(const FArguments& InArgs)
 {
+	
+
 	TakeRecorderSubsystem = GEngine->GetEngineSubsystem<UTakeRecorderSubsystem>();
 	
 	FTakeRecorderSequenceParameters Data;
@@ -103,7 +109,8 @@ void STakeRecorderPanel::Construct(const FArguments& InArgs)
 	LevelSequenceTakeWidget = SNew(SLevelSequenceTakeEditor)
 	.LevelSequence(this, &STakeRecorderPanel::GetLevelSequence)
 	.OnDetailsPropertiesChanged(this, &STakeRecorderPanel::OnLevelSequenceDetailsChanged)
-	.OnDetailsViewAdded(this, &STakeRecorderPanel::OnLevelSequenceDetailsViewAdded);
+	.OnDetailsViewAdded(this, &STakeRecorderPanel::OnLevelSequenceDetailsViewAdded)
+	.ShouldPromptUserForTargetRecordClassChange(this, &STakeRecorderPanel::ShouldPromptUserForTargetRecordClassChange);
 
 	// Create the sequencer panel, and open it if necessary
 	SequencerPanel = MakeShared<FScopedSequencerPanel>(MakeAttributeSP(this, &STakeRecorderPanel::GetLevelSequence));
@@ -745,6 +752,42 @@ FReply STakeRecorderPanel::OnReviewLastRecording()
 	return FReply::Handled();
 }
 
+void STakeRecorderPanel::OnTakePresetSettingsChanged()
+{
+	// Settings may have been changed by undo / redo... in that case undo / redo will have already updated the internal level sequence.
+	if (GIsTransacting)
+	{
+		// OnBeforeRedoUndo has closed the panel - we'll now reinitialize it with the level sequence.
+		RefreshPanel();
+	}
+	else
+	{
+		ClearPendingTake();
+	}
+}
+
+void STakeRecorderPanel::OnBeforeRedoUndo(const FTransactionContext& TransactionContext) const
+{
+	UTransBuffer* Transactor = GEditor ? Cast<UTransBuffer>(GEditor->Trans) : nullptr;
+	if (!Transactor)
+	{
+		return;
+	}
+	
+	const int32 Index = Transactor->FindTransactionIndex(TransactionContext.TransactionId);
+	const FTransaction* Transaction = Transactor->GetTransaction(Index);
+	const bool bAffectsSettings = Transaction && Transaction->ContainsObject(UTakePresetSettings::Get());
+
+	// Sequencer does not deal well with the underlying ULevelSequence class changing, i.e. when UTakePresetSettings::RecordTargetClass is changed.
+	// Checks & ensures fly. After undo / redo, it seems like FSharedPlaybackState::WeakRootSequence is not transacted properly.
+	// The proper fix would be to find out why that happens. However, right now it's easier to reinitialize sequencer by closing the tab and
+	// reopening it post undo in OnTakePresetSettingsChanged.
+	if (bAffectsSettings)
+	{
+		SequencerPanel->Close();
+	}
+}
+
 FReply STakeRecorderPanel::OnRevertChanges()
 {
 	FText WarningMessage(LOCTEXT("Warning_RevertChanges", "Are you sure you want to revert changes? Your current changes will be discarded."));
@@ -869,6 +912,19 @@ void STakeRecorderPanel::OnLevelSequenceDetailsViewAdded(const TWeakPtr<IDetails
 	}
 }
 
+SLevelSequenceTakeEditor::EShowPrompt STakeRecorderPanel::ShouldPromptUserForTargetRecordClassChange(const UClass* NewClass) const
+{
+	UMovieScene* MovieScene = GetLevelSequence() ? GetLevelSequence()->GetMovieScene() : nullptr;
+	const bool bHasUserMadeChanges = MovieScene
+		&& (!MovieScene->GetTracks().IsEmpty()
+			|| !MovieScene->GetBindings().IsEmpty()
+			|| MovieScene->GetPossessableCount() > 0
+			|| MovieScene->GetSpawnableCount() > 0);
+	return bHasUserMadeChanges
+		? SLevelSequenceTakeEditor::EShowPrompt::ShowPrompt
+		: SLevelSequenceTakeEditor::EShowPrompt::AllowOperationSilently; 
+}
+
 void STakeRecorderPanel::OnRecordingInitialized(UTakeRecorder* Recorder)
 {
 	// It's important that UTakeRecorderEditorSubsystem::OnRecordingInitialized has fired before this point,
diff --git a/Engine/Plugins/VirtualProduction/Takes/Source/TakeRecorder/Private/Widgets/STakeRecorderPanel.h b/Engine/Plugins/VirtualProduction/Takes/Source/TakeRecorder/Private/Widgets/STakeRecorderPanel.h
--- a/Engine/Plugins/VirtualProduction/Takes/Source/TakeRecorder/Private/Widgets/STakeRecorderPanel.h
+++ b/Engine/Plugins/VirtualProduction/Takes/Source/TakeRecorder/Private/Widgets/STakeRecorderPanel.h
@@ -2,8 +2,10 @@
 
 #pragma once
 
+#include "EditorUndoClient.h"
 #include "Widgets/DeclarativeSyntaxSupport.h"
 #include "Widgets/SCompoundWidget.h"
+#include "SLevelSequenceTakeEditor.h"
 #include "UObject/GCObject.h"
 #include "Recorder/TakeRecorderParameters.h"
 
@@ -103,6 +105,11 @@ private:
 
 	FReply OnReviewLastRecording();
 
+	/** Handles the UTakePresetSettings::RecordTargetClass changing. Recreates the object if recording a transaction or refreshes sequencer if undo / redoing. */
+	void OnTakePresetSettingsChanged();
+	/** If the transaction changes the UTakePresetSettings::RecordTargetClass, temporarily closes the sequencer in order to refresh it. */
+	void OnBeforeRedoUndo(const FTransactionContext& TransactionContext) const;
+
 	TSharedRef<SWidget> OnOpenSequenceToRecordIntoMenu();
 
 	void OnOpenSequenceToRecordInto(const FAssetData& InAsset);
@@ -115,6 +122,11 @@ private:
 	void OnLevelSequenceDetailsChanged(const FPropertyChangedEvent& InPropertyChangedEvent);
 	/** When a level sequence has a details view added. */
 	void OnLevelSequenceDetailsViewAdded(const TWeakPtr<IDetailsView>& InDetailsView);
+	/**
+	 * Called when the user attempts to change UTakePresetSettings::TargetRecordClass.
+	 * @return Whether it is necessary to show a prompt to the user.
+	 */
+	SLevelSequenceTakeEditor::EShowPrompt ShouldPromptUserForTargetRecordClassChange(const UClass* NewClass) const;
 
 	ECheckBoxState GetTakeBrowserCheckState() const;
 	void ToggleTakeBrowserCheckState(ECheckBoxState CheckState);
diff --git a/Engine/Plugins/VirtualProduction/Takes/Source/TakeRecorder/TakeRecorder.Build.cs b/Engine/Plugins/VirtualProduction/Takes/Source/TakeRecorder/TakeRecorder.Build.cs
--- a/Engine/Plugins/VirtualProduction/Takes/Source/TakeRecorder/TakeRecorder.Build.cs
+++ b/Engine/Plugins/VirtualProduction/Takes/Source/TakeRecorder/TakeRecorder.Build.cs
@@ -11,6 +11,7 @@ public class TakeRecorder : ModuleRules
 				"AppFramework",
 				"AssetDefinition",
 				"AssetRegistry",
+				"ClassViewer",
 				"ContentBrowser",
 				"Core",
 				"CoreUObject",
diff --git a/Engine/Plugins/VirtualProduction/Takes/Source/TakesCore/Private/TakePreset.cpp b/Engine/Plugins/VirtualProduction/Takes/Source/TakesCore/Private/TakePreset.cpp
--- a/Engine/Plugins/VirtualProduction/Takes/Source/TakesCore/Private/TakePreset.cpp
+++ b/Engine/Plugins/VirtualProduction/Takes/Source/TakesCore/Private/TakePreset.cpp
@@ -4,6 +4,7 @@
 #include "LevelSequence.h"
 #include "MovieScene.h"
 #include "MovieSceneToolsProjectSettings.h"
+#include "TakePresetSettings.h"
 #include "Misc/TransactionObjectEvent.h"
 #include "UObject/Package.h"
 
@@ -68,7 +69,7 @@ void UTakePreset::CreateLevelSequence()
 
 	// Copy the transient and transactional flags from the parent
 	EObjectFlags SequenceFlags = GetFlags() & (RF_Transient | RF_Transactional);
-	LevelSequence = NewObject<ULevelSequence>(this, GetFName(), SequenceFlags);
+	LevelSequence = NewObject<ULevelSequence>(this, UTakePresetSettings::Get()->GetTargetRecordClass(), GetFName(), SequenceFlags);
 	LevelSequence->Initialize();
 
 	const UMovieSceneToolsProjectSettings* ProjectSettings = GetDefault<UMovieSceneToolsProjectSettings>();
diff --git a/Engine/Plugins/VirtualProduction/Takes/Source/TakesCore/Private/TakePresetSettings.cpp b/Engine/Plugins/VirtualProduction/Takes/Source/TakesCore/Private/TakePresetSettings.cpp
new file mode 100644
--- /dev/null
+++ b/Engine/Plugins/VirtualProduction/Takes/Source/TakesCore/Private/TakePresetSettings.cpp
@@ -0,0 +1,27 @@
+﻿// Copyright Epic Games, Inc. All Rights Reserved.
+
+#include "TakePresetSettings.h"
+
+#include "LevelSequence.h"
+
+UTakePresetSettings::UTakePresetSettings()
+	: TargetRecordClass(ULevelSequence::StaticClass())
+{}
+
+UTakePresetSettings* UTakePresetSettings::Get()
+{
+	return GetMutableDefault<UTakePresetSettings>();
+}
+
+UClass* UTakePresetSettings::GetTargetRecordClass() const
+{
+	return TargetRecordClass.TargetRecordClass ? TargetRecordClass.TargetRecordClass.Get() : ULevelSequence::StaticClass();
+}
+
+#if WITH_EDITOR
+void UTakePresetSettings::PostEditChangeProperty(FPropertyChangedEvent& PropertyChangedEvent)
+{
+	UObject::PostEditChangeProperty(PropertyChangedEvent);
+	OnSettingsChangedDelegate.Broadcast();
+}
+#endif
diff --git a/Engine/Plugins/VirtualProduction/Takes/Source/TakesCore/Public/TakePresetSettings.h b/Engine/Plugins/VirtualProduction/Takes/Source/TakesCore/Public/TakePresetSettings.h
new file mode 100644
--- /dev/null
+++ b/Engine/Plugins/VirtualProduction/Takes/Source/TakesCore/Public/TakePresetSettings.h
@@ -0,0 +1,59 @@
+﻿// Copyright Epic Games, Inc. All Rights Reserved.
+
+#pragma once
+
+#include "Templates/SubclassOf.h"
+#include "UObject/Object.h"
+#include "TakePresetSettings.generated.h"
+
+class ULevelSequence;
+
+/** Wraps the class so this can be customized by a type layout. */
+USTRUCT()
+struct FTakeRecorderTargetRecordClassProperty
+{
+	GENERATED_BODY()
+
+	UPROPERTY(EditAnywhere, Category = "Take Recorder")
+	TSubclassOf<ULevelSequence> TargetRecordClass;
+};
+
+/**
+ * Settings for how to assemble UTakePreset.
+ * 
+ * These settings belong into UTakeRecorderSettings but those are in TakeRecorder, which depends on TakesCore. 
+ * DisplayName is important so these settings are displayed in the same category as UTakeRecorderSettings.
+ */
+UCLASS(config=EditorSettings, BlueprintType, DisplayName = "Take Recorder")
+class TAKESCORE_API UTakePresetSettings : public UObject
+{
+public:
+	GENERATED_BODY()
+
+	UTakePresetSettings();
+
+	/** @return The settings object */
+	static UTakePresetSettings* Get();
+
+	/** @return The class that recorded sequences should have */
+	UClass* GetTargetRecordClass() const;
+
+	DECLARE_MULTICAST_DELEGATE(FOnSettingsChanged);
+	/** Invoked when any settings change. */
+	FOnSettingsChanged& OnSettingsChanged() { return OnSettingsChangedDelegate; }
+
+	static FName GetTargetRecordClassMemberName() { return GET_MEMBER_NAME_CHECKED(UTakePresetSettings, TargetRecordClass); } 
+
+#if WITH_EDITOR
+	virtual void PostEditChangeProperty(FPropertyChangedEvent& PropertyChangedEvent) override;
+#endif
+
+private:
+	
+	/** The class that recorded sequences should have. */
+	UPROPERTY(config, EditAnywhere, Category = "Take Recorder")
+	FTakeRecorderTargetRecordClassProperty TargetRecordClass;
+
+	/** Invoked when any settings change. */
+	FOnSettingsChanged OnSettingsChangedDelegate;
+};
