Add missing code that was accidentally  removed during submit operation

diff --git a/Engine/Plugins/VirtualProduction/Takes/Source/TakeRecorder/Private/Widgets/STakeRecorderPanel.cpp b/Engine/Plugins/VirtualProduction/Takes/Source/TakeRecorder/Private/Widgets/STakeRecorderPanel.cpp
--- a/Engine/Plugins/VirtualProduction/Takes/Source/TakeRecorder/Private/Widgets/STakeRecorderPanel.cpp
+++ b/Engine/Plugins/VirtualProduction/Takes/Source/TakeRecorder/Private/Widgets/STakeRecorderPanel.cpp
@@ -91,7 +91,12 @@ STakeRecorderPanel::~STakeRecorderPanel()
 UE_DISABLE_OPTIMIZATION_SHIP
 void STakeRecorderPanel::Construct(const FArguments& InArgs)
 {
-	
+	// If the UTakePresetSettings::TargetRecordClass changes, the level sequence in the TransientPreset must be regenerated.
+	UTakePresetSettings::Get()->OnSettingsChanged().AddSP(this, &STakeRecorderPanel::OnTakePresetSettingsChanged);
+	if (UTransBuffer* Transactor = GEditor ? Cast<UTransBuffer>(GEditor->Trans) : nullptr)
+	{
+		Transactor->OnBeforeRedoUndo().AddSP(this, &STakeRecorderPanel::OnBeforeRedoUndo);
+	}
 
 	TakeRecorderSubsystem = GEngine->GetEngineSubsystem<UTakeRecorderSubsystem>();
 	
