Fix TargetRecordClass not saving

diff --git a/Engine/Plugins/VirtualProduction/Takes/Source/TakesCore/Private/TakePresetSettings.cpp b/Engine/Plugins/VirtualProduction/Takes/Source/TakesCore/Private/TakePresetSettings.cpp
--- a/Engine/Plugins/VirtualProduction/Takes/Source/TakesCore/Private/TakePresetSettings.cpp
+++ b/Engine/Plugins/VirtualProduction/Takes/Source/TakesCore/Private/TakePresetSettings.cpp
@@ -22,6 +22,7 @@ UClass* UTakePresetSettings::GetTargetRecordClass() const
 void UTakePresetSettings::PostEditChangeProperty(FPropertyChangedEvent& PropertyChangedEvent)
 {
 	UObject::PostEditChangeProperty(PropertyChangedEvent);
+	SaveConfig();
 	OnSettingsChangedDelegate.Broadcast();
 }
 #endif
diff --git a/Engine/Plugins/VirtualProduction/Takes/Source/TakesCore/Public/TakePresetSettings.h b/Engine/Plugins/VirtualProduction/Takes/Source/TakesCore/Public/TakePresetSettings.h
--- a/Engine/Plugins/VirtualProduction/Takes/Source/TakesCore/Public/TakePresetSettings.h
+++ b/Engine/Plugins/VirtualProduction/Takes/Source/TakesCore/Public/TakePresetSettings.h
@@ -15,7 +15,7 @@ struct FTakeRecorderTargetRecordClassProperty
 	GENERATED_BODY()
 
 	UPROPERTY(EditAnywhere, Category = "Take Recorder")
-	TSubclassOf<ULevelSequence> TargetRecordClass;
+	TSoftClassPtr<ULevelSequence> TargetRecordClass;
 };
 
 /**
