Merge logic into new UTakeRecorderSubsytem.
For context, this logic was developed at the same time as the UTakeRecorderSubsytem refactor. This is a separate CL to merge both changes.

diff --git a/Engine/Plugins/VirtualProduction/Takes/Source/TakeRecorder/Private/Recorder/TakeRecorderSubsystem.cpp b/Engine/Plugins/VirtualProduction/Takes/Source/TakeRecorder/Private/Recorder/TakeRecorderSubsystem.cpp
--- a/Engine/Plugins/VirtualProduction/Takes/Source/TakeRecorder/Private/Recorder/TakeRecorderSubsystem.cpp
+++ b/Engine/Plugins/VirtualProduction/Takes/Source/TakeRecorder/Private/Recorder/TakeRecorderSubsystem.cpp
@@ -4,6 +4,7 @@
 
 #include "AssetRegistry/AssetRegistryModule.h"
 #include "Dialog/SMessageDialog.h"
+#include "Editor/TransBuffer.h"
 #include "Framework/Notifications/NotificationManager.h"
 #include "ILevelSequenceEditorToolkit.h"
 #include "ITakeRecorderModule.h"
@@ -17,6 +18,7 @@
 #include "Recorder/TakeRecorderParameters.h"
 #include "ScopedTransaction.h"
 #include "TakeMetaData.h"
+#include "TakePresetSettings.h"
 #include "TakeRecorderModule.h"
 #include "TakeRecorderSettings.h"
 #include "TakeRecorderSource.h"
@@ -42,6 +44,9 @@ void UTakeRecorderSubsystem::Initialize(FSubsystemCollectionBase& Collection)
 	NamingTokensData->SetFlags(RF_Transactional);
 
 	BindToNamingTokenEvents();
+	
+	// If the UTakePresetSettings::TargetRecordClass changes, the level sequence in the TransientPreset must be regenerated.
+	UTakePresetSettings::Get()->OnSettingsChanged().AddUObject(this, &UTakeRecorderSubsystem::OnTakePresetSettingsChanged);
 }
 
 void UTakeRecorderSubsystem::Deinitialize()
@@ -763,6 +768,17 @@ UTakeRecorderNamingTokensData* UTakeRecorderSubsystem::GetNamingTokensData() con
 	return NamingTokensData;
 }
 
+bool UTakeRecorderSubsystem::HasPendingChanges() const
+{
+	UMovieScene* MovieScene = GetLevelSequence() ? GetLevelSequence()->GetMovieScene() : nullptr;
+	const bool bHasUserMadeChanges = MovieScene
+		&& (!MovieScene->GetTracks().IsEmpty()
+			|| !MovieScene->GetBindings().IsEmpty()
+			|| MovieScene->GetPossessableCount() > 0
+			|| MovieScene->GetSpawnableCount() > 0);
+	return bHasUserMadeChanges;
+}
+
 FOnTakeRecordingInitialized& UTakeRecorderSubsystem::GetOnRecordingPreInitializedEvent()
 {
 	return OnRecordingPreInitializeEvent;
@@ -1034,4 +1050,15 @@ void UTakeRecorderSubsystem::OnTakeRecorderNamingTokensPreEvaluate(const FNaming
 	}
 }
 
+void UTakeRecorderSubsystem::OnTakePresetSettingsChanged()
+{
+	// This effectively regenerates the level sequence in the preset.
+	// Existing data is discarded. In the future, we should try to migrate it.
+	// FTakePresetRecorderCustomization handles asking the user whether it is ok to discard the changes. Once this event fires, the decision has been made.
+	if (!GIsTransacting)
+	{
+		ClearPendingTake();
+	}
+}
+
 #undef LOCTEXT_NAMESPACE
diff --git a/Engine/Plugins/VirtualProduction/Takes/Source/TakeRecorder/Private/Widgets/SLevelSequenceTakeEditor.cpp b/Engine/Plugins/VirtualProduction/Takes/Source/TakeRecorder/Private/Widgets/SLevelSequenceTakeEditor.cpp
--- a/Engine/Plugins/VirtualProduction/Takes/Source/TakeRecorder/Private/Widgets/SLevelSequenceTakeEditor.cpp
+++ b/Engine/Plugins/VirtualProduction/Takes/Source/TakeRecorder/Private/Widgets/SLevelSequenceTakeEditor.cpp
@@ -58,6 +58,7 @@
 #include "PropertyEditorModule.h"
 #include "IDetailsView.h"
 #include "DetailWidgetRow.h"
+#include "Recorder/TakeRecorderSubsystem.h"
 
 #define LOCTEXT_NAMESPACE "SLevelSequenceTakeEditor"
 
@@ -105,7 +106,6 @@ void SLevelSequenceTakeEditor::Construct(const FArguments& InArgs)
 
 	OnDetailsPropertiesChangedEvent = InArgs._OnDetailsPropertiesChanged;
 	OnDetailsViewAddedEvent = InArgs._OnDetailsViewAdded;
-	ShouldPromptUserForTargetRecordClassChangeEvent = InArgs._ShouldPromptUserForTargetRecordClassChange;
 
 	DetailsBox = SNew(SScrollBox);
 	DetailsBox->SetScrollBarRightClickDragAllowed(true);
@@ -288,8 +288,9 @@ void SLevelSequenceTakeEditor::OnDetailsPropertiesChanged(const FPropertyChanged
 
 bool SLevelSequenceTakeEditor::PromptUserForTargetRecordClassChange(const UClass* NewClass) const
 {
-	if (ShouldPromptUserForTargetRecordClassChangeEvent.IsBound()
-		&& ShouldPromptUserForTargetRecordClassChangeEvent.Execute(NewClass) == EShowPrompt::AllowOperationSilently)
+	UTakeRecorderSubsystem* Subsystem = GEngine ? GEngine->GetEngineSubsystem<UTakeRecorderSubsystem>() : nullptr;
+	const bool bAllowSilently = Subsystem && !Subsystem->HasPendingChanges();
+	if (bAllowSilently)
 	{
 		return true;
 	}
diff --git a/Engine/Plugins/VirtualProduction/Takes/Source/TakeRecorder/Private/Widgets/SLevelSequenceTakeEditor.h b/Engine/Plugins/VirtualProduction/Takes/Source/TakeRecorder/Private/Widgets/SLevelSequenceTakeEditor.h
--- a/Engine/Plugins/VirtualProduction/Takes/Source/TakeRecorder/Private/Widgets/SLevelSequenceTakeEditor.h
+++ b/Engine/Plugins/VirtualProduction/Takes/Source/TakeRecorder/Private/Widgets/SLevelSequenceTakeEditor.h
@@ -28,9 +28,6 @@ public:
 	DECLARE_DELEGATE_OneParam(FOnDetailsPropertiesChanged, const FPropertyChangedEvent&);
 	DECLARE_DELEGATE_OneParam(FOnDetailsViewAdded, const TWeakPtr<IDetailsView>&);
 
-	enum class EShowPrompt{ ShowPrompt, AllowOperationSilently };
-	DECLARE_DELEGATE_RetVal_OneParam(EShowPrompt, FShouldPromptUserForTargetRecordClassChange, const UClass* /*NewClass*/)
-
 	SLATE_BEGIN_ARGS(SLevelSequenceTakeEditor)
 		: _LevelSequence(nullptr)
 		{}
@@ -38,8 +35,6 @@ public:
 		SLATE_ATTRIBUTE(ULevelSequence*, LevelSequence)
 		SLATE_EVENT(FOnDetailsPropertiesChanged, OnDetailsPropertiesChanged)
 		SLATE_EVENT(FOnDetailsViewAdded, OnDetailsViewAdded)
-		/** Called when the user attempts to change UTakePresetSettings::TargetRecordClass. Returns whether it is necessary to show a prompt to the user. */
-		SLATE_EVENT(FShouldPromptUserForTargetRecordClassChange, ShouldPromptUserForTargetRecordClassChange)
 
 	SLATE_END_ARGS()
 
@@ -117,6 +112,4 @@ private:
 	FOnDetailsPropertiesChanged OnDetailsPropertiesChangedEvent;
 	/** Called when a details view has been added. */
 	FOnDetailsViewAdded OnDetailsViewAddedEvent;
-	/** Called when the user attempts to change UTakePresetSettings::TargetRecordClass. Returns whether it is necessary to show a prompt to the user. */
-	FShouldPromptUserForTargetRecordClassChange ShouldPromptUserForTargetRecordClassChangeEvent;
 };
\ No newline at end of file
diff --git a/Engine/Plugins/VirtualProduction/Takes/Source/TakeRecorder/Private/Widgets/STakeRecorderPanel.cpp b/Engine/Plugins/VirtualProduction/Takes/Source/TakeRecorder/Private/Widgets/STakeRecorderPanel.cpp
--- a/Engine/Plugins/VirtualProduction/Takes/Source/TakeRecorder/Private/Widgets/STakeRecorderPanel.cpp
+++ b/Engine/Plugins/VirtualProduction/Takes/Source/TakeRecorder/Private/Widgets/STakeRecorderPanel.cpp
@@ -91,8 +91,6 @@ STakeRecorderPanel::~STakeRecorderPanel()
 UE_DISABLE_OPTIMIZATION_SHIP
 void STakeRecorderPanel::Construct(const FArguments& InArgs)
 {
-	// If the UTakePresetSettings::TargetRecordClass changes, the level sequence in the TransientPreset must be regenerated.
-	UTakePresetSettings::Get()->OnSettingsChanged().AddSP(this, &STakeRecorderPanel::OnTakePresetSettingsChanged);
 	if (UTransBuffer* Transactor = GEditor ? Cast<UTransBuffer>(GEditor->Trans) : nullptr)
 	{
 		Transactor->OnBeforeRedoUndo().AddSP(this, &STakeRecorderPanel::OnBeforeRedoUndo);
@@ -114,8 +112,7 @@ void STakeRecorderPanel::Construct(const FArguments& InArgs)
 	LevelSequenceTakeWidget = SNew(SLevelSequenceTakeEditor)
 	.LevelSequence(this, &STakeRecorderPanel::GetLevelSequence)
 	.OnDetailsPropertiesChanged(this, &STakeRecorderPanel::OnLevelSequenceDetailsChanged)
-	.OnDetailsViewAdded(this, &STakeRecorderPanel::OnLevelSequenceDetailsViewAdded)
-	.ShouldPromptUserForTargetRecordClassChange(this, &STakeRecorderPanel::ShouldPromptUserForTargetRecordClassChange);
+	.OnDetailsViewAdded(this, &STakeRecorderPanel::OnLevelSequenceDetailsViewAdded);
 
 	// Create the sequencer panel, and open it if necessary
 	SequencerPanel = MakeShared<FScopedSequencerPanel>(MakeAttributeSP(this, &STakeRecorderPanel::GetLevelSequence));
@@ -765,10 +762,6 @@ void STakeRecorderPanel::OnTakePresetSettingsChanged()
 		// OnBeforeRedoUndo has closed the panel - we'll now reinitialize it with the level sequence.
 		RefreshPanel();
 	}
-	else
-	{
-		ClearPendingTake();
-	}
 }
 
 void STakeRecorderPanel::OnBeforeRedoUndo(const FTransactionContext& TransactionContext) const
@@ -917,19 +910,6 @@ void STakeRecorderPanel::OnLevelSequenceDetailsViewAdded(const TWeakPtr<IDetails
 	}
 }
 
-SLevelSequenceTakeEditor::EShowPrompt STakeRecorderPanel::ShouldPromptUserForTargetRecordClassChange(const UClass* NewClass) const
-{
-	UMovieScene* MovieScene = GetLevelSequence() ? GetLevelSequence()->GetMovieScene() : nullptr;
-	const bool bHasUserMadeChanges = MovieScene
-		&& (!MovieScene->GetTracks().IsEmpty()
-			|| !MovieScene->GetBindings().IsEmpty()
-			|| MovieScene->GetPossessableCount() > 0
-			|| MovieScene->GetSpawnableCount() > 0);
-	return bHasUserMadeChanges
-		? SLevelSequenceTakeEditor::EShowPrompt::ShowPrompt
-		: SLevelSequenceTakeEditor::EShowPrompt::AllowOperationSilently; 
-}
-
 void STakeRecorderPanel::OnRecordingInitialized(UTakeRecorder* Recorder)
 {
 	// It's important that UTakeRecorderEditorSubsystem::OnRecordingInitialized has fired before this point,
diff --git a/Engine/Plugins/VirtualProduction/Takes/Source/TakeRecorder/Private/Widgets/STakeRecorderPanel.h b/Engine/Plugins/VirtualProduction/Takes/Source/TakeRecorder/Private/Widgets/STakeRecorderPanel.h
--- a/Engine/Plugins/VirtualProduction/Takes/Source/TakeRecorder/Private/Widgets/STakeRecorderPanel.h
+++ b/Engine/Plugins/VirtualProduction/Takes/Source/TakeRecorder/Private/Widgets/STakeRecorderPanel.h
@@ -122,11 +122,6 @@ private:
 	void OnLevelSequenceDetailsChanged(const FPropertyChangedEvent& InPropertyChangedEvent);
 	/** When a level sequence has a details view added. */
 	void OnLevelSequenceDetailsViewAdded(const TWeakPtr<IDetailsView>& InDetailsView);
-	/**
-	 * Called when the user attempts to change UTakePresetSettings::TargetRecordClass.
-	 * @return Whether it is necessary to show a prompt to the user.
-	 */
-	SLevelSequenceTakeEditor::EShowPrompt ShouldPromptUserForTargetRecordClassChange(const UClass* NewClass) const;
 
 	ECheckBoxState GetTakeBrowserCheckState() const;
 	void ToggleTakeBrowserCheckState(ECheckBoxState CheckState);
diff --git a/Engine/Plugins/VirtualProduction/Takes/Source/TakeRecorder/Public/Recorder/TakeRecorderSubsystem.h b/Engine/Plugins/VirtualProduction/Takes/Source/TakeRecorder/Public/Recorder/TakeRecorderSubsystem.h
--- a/Engine/Plugins/VirtualProduction/Takes/Source/TakeRecorder/Public/Recorder/TakeRecorderSubsystem.h
+++ b/Engine/Plugins/VirtualProduction/Takes/Source/TakeRecorder/Public/Recorder/TakeRecorderSubsystem.h
@@ -39,6 +39,7 @@ class TAKERECORDER_API UTakeRecorderSubsystem : public UEngineSubsystem, public
 	GENERATED_BODY()
 
 public:
+	
 	virtual void Initialize(FSubsystemCollectionBase& Collection) override;
 	virtual void Deinitialize() override;
 
@@ -275,6 +276,9 @@ public:
 
 	/** Retrieve relevant Naming Tokens data for Take Recorder. */
 	UTakeRecorderNamingTokensData* GetNamingTokensData() const;
+
+	/** @return Whether there are any pending changes, which ClearPendingTake would discard. */
+	bool HasPendingChanges() const;
 	
 	/**
 	 * Retrieve a multi-cast delegate that is triggered when a recording pre-initializes
@@ -349,6 +353,9 @@ private:
 	/** Called before our naming tokens class evaluates. Used to populate naming token data with our custom token definitions. */
 	void OnTakeRecorderNamingTokensPreEvaluate(const FNamingTokensEvaluationData& InEvaluationData);
 	
+	/** Handles the UTakePresetSettings::RecordTargetClass changing: Recreates the transient level sequence if recording a transaction. */
+	void OnTakePresetSettingsChanged();
+	
 private:
 	/** Last data used for initialization. */
 	FTakeRecorderSequenceParameters TargetSequenceData;
