TakeRecorder: Add takeName token, representing the sequence asset name.
Default animation and audio sources to use the takeName token.
Add validation for generating a root directory path.
Fix TR user tokens adding in slate and take, which would show up twice in Shots token display.

NamingTokens: Add utility function to check for existence of a token. Export other utility functions.

diff --git a/Engine/Plugins/Developer/NamingTokens/Source/NamingTokens/Private/Utils/NamingTokenUtils.cpp b/Engine/Plugins/Developer/NamingTokens/Source/NamingTokens/Private/Utils/NamingTokenUtils.cpp
--- a/Engine/Plugins/Developer/NamingTokens/Source/NamingTokens/Private/Utils/NamingTokenUtils.cpp
+++ b/Engine/Plugins/Developer/NamingTokens/Source/NamingTokens/Private/Utils/NamingTokenUtils.cpp
@@ -7,6 +7,14 @@
 #include "Internationalization/LocKeyFuncs.h"
 #include "Internationalization/Regex.h"
 
+namespace UE::NamingTokens::Utils::Private
+{
+	FString GetTokenPatternString()
+	{
+		return FString::Printf(TEXT(R"(\{\s*((?:[a-zA-Z0-9_]+%s)*[a-zA-Z0-9_]+)\s*\})"), *GetNamespaceDelimiter());
+	}
+}
+
 UFunction* UE::NamingTokens::Utils::GetProcessTokenFunctionSignature()
 {
 	UFunction* FunctionSignature = UNamingTokens::StaticClass()->FindFunctionByName(UNamingTokens::GetProcessTokenTemplateFunctionName());
@@ -39,7 +47,7 @@ TArray<FString> UE::NamingTokens::Utils::GetTokenKeysFromString(const FString& I
 {
 	TSet<FString, FLocKeySetFuncs> Tokens;
 
-	const FString PatternString = FString::Printf(TEXT(R"(\{\s*((?:[a-zA-Z0-9_]+%s)*[a-zA-Z0-9_]+)\s*\})"), *GetNamespaceDelimiter());
+	const FString PatternString = Private::GetTokenPatternString();
 	const FRegexPattern Pattern(PatternString);
 	FRegexMatcher Matcher(Pattern, InTokenString);
 
@@ -52,6 +60,25 @@ TArray<FString> UE::NamingTokens::Utils::GetTokenKeysFromString(const FString& I
 	return Tokens.Array();
 }
 
+bool UE::NamingTokens::Utils::IsTokenInString(const FString& InTokenKey, const FString& InTokenString)
+{
+	const FString PatternString = Private::GetTokenPatternString();
+	const FRegexPattern Pattern(PatternString);
+	FRegexMatcher Matcher(Pattern, InTokenString);
+
+	while (Matcher.FindNext())
+	{
+		const FString Token = Matcher.GetCaptureGroup(1);
+		const FString TokenNoNamespace = RemoveNamespaceFromTokenKey(Token);
+		if (TokenNoNamespace.Equals(InTokenKey))
+		{
+			return true;
+		}
+	}
+
+	return false;
+}
+
 FString UE::NamingTokens::Utils::GetNamespaceFromTokenKey(const FString& InTokenKey)
 {
 	FString LeftPart, RightPart;
diff --git a/Engine/Plugins/Developer/NamingTokens/Source/NamingTokens/Public/Utils/NamingTokenUtils.h b/Engine/Plugins/Developer/NamingTokens/Source/NamingTokens/Public/Utils/NamingTokenUtils.h
--- a/Engine/Plugins/Developer/NamingTokens/Source/NamingTokens/Public/Utils/NamingTokenUtils.h
+++ b/Engine/Plugins/Developer/NamingTokens/Source/NamingTokens/Public/Utils/NamingTokenUtils.h
@@ -21,19 +21,22 @@ namespace UE::NamingTokens::Utils
 	NAMINGTOKENS_API FString GetNamespaceDelimiter();
 	
 	/** Extract all token keys from a string. */
-	TArray<FString> GetTokenKeysFromString(const FString& InTokenString);
+	NAMINGTOKENS_API TArray<FString> GetTokenKeysFromString(const FString& InTokenString);
 
+	/** Checks if a token is present in a string. */
+	NAMINGTOKENS_API bool IsTokenInString(const FString& InTokenKey, const FString& InTokenString);
+	
 	/**
 	 * Return the namespace from a token key, or an empty string.
 	 * @param InTokenKey The token key shouldn't have any additional formatting beyond the namespace, eg "namespace.token"
 	 * @return The namespace, or an empty string.
 	 */
-	FString GetNamespaceFromTokenKey(const FString& InTokenKey);
+	NAMINGTOKENS_API FString GetNamespaceFromTokenKey(const FString& InTokenKey);
 	
 	/**
 	 * Remove the namespace from the given token key.
 	 * @param InTokenKey The token key shouldn't have any additional formatting beyond the namespace, eg "namespace.token"
 	 * @return The token key without the namespace, or the InTokenKey if there was no namespace.
 	 */
-	FString RemoveNamespaceFromTokenKey(const FString& InTokenKey);
+	NAMINGTOKENS_API FString RemoveNamespaceFromTokenKey(const FString& InTokenKey);
 }
\ No newline at end of file
diff --git a/Engine/Plugins/VirtualProduction/Takes/Source/TakeRecorder/Private/Recorder/TakeRecorder.cpp b/Engine/Plugins/VirtualProduction/Takes/Source/TakeRecorder/Private/Recorder/TakeRecorder.cpp
--- a/Engine/Plugins/VirtualProduction/Takes/Source/TakeRecorder/Private/Recorder/TakeRecorder.cpp
+++ b/Engine/Plugins/VirtualProduction/Takes/Source/TakeRecorder/Private/Recorder/TakeRecorder.cpp
@@ -619,7 +619,12 @@ bool UTakeRecorder::CreateDestinationAsset(const TCHAR* AssetPathFormat, ULevelS
 {
 	check(LevelSequenceBase && Sources && MetaData);
 
-	FString   PackageName = MetaData->GenerateAssetPath(AssetPathFormat);
+	FString   PackageName;
+
+	if (!MetaData->TryGenerateRootAssetPath(AssetPathFormat, PackageName))
+	{
+		return false;
+	}
 
 	// Initialize a new package, ensuring that it has a unique name
 	if (!TakesUtils::CreateNewAssetPackage<ULevelSequence>(PackageName, SequenceAsset, OutError, LevelSequenceBase))
@@ -839,25 +844,28 @@ void UTakeRecorder::PreRecord()
 	MovieSceneSerializationNamespace::bAutoSerialize = Parameters.User.bAutoSerialize;
 	if (Parameters.User.bAutoSerialize)
 	{
-		FString AssetName = AssetMetaData->GenerateAssetPath(Parameters.Project.GetTakeAssetPath());
-		FString AssetPath = FPaths::ProjectSavedDir() + AssetName;
-		FPaths::RemoveDuplicateSlashes(AssetPath);
-		IPlatformFile& PlatformFile = FPlatformFileManager::Get().GetPlatformFile();
-		if (!PlatformFile.DirectoryExists(*AssetPath))
+		FString AssetName;
+		if (AssetMetaData->TryGenerateRootAssetPath(Parameters.Project.GetTakeAssetPath(), AssetName))
 		{
-			PlatformFile.CreateDirectoryTree(*AssetPath);
-		}
+			FString AssetPath = FPaths::ProjectSavedDir() + AssetName;
+			FPaths::RemoveDuplicateSlashes(AssetPath);
+			IPlatformFile& PlatformFile = FPlatformFileManager::Get().GetPlatformFile();
+			if (!PlatformFile.DirectoryExists(*AssetPath))
+			{
+				PlatformFile.CreateDirectoryTree(*AssetPath);
+			}
 
-		ManifestSerializer.SetLocalCaptureDir(AssetPath);
-		FName SerializedType("Sequence");
-		FString Name = SequenceAsset->GetName();
-		FManifestFileHeader Header(Name, SerializedType, FGuid());
-		FText Error;
-		FString FileName = FString::Printf(TEXT("%s_%s"), *(SerializedType.ToString()), *(Name));
+			ManifestSerializer.SetLocalCaptureDir(AssetPath);
+			FName SerializedType("Sequence");
+			FString Name = SequenceAsset->GetName();
+			FManifestFileHeader Header(Name, SerializedType, FGuid());
+			FText Error;
+			FString FileName = FString::Printf(TEXT("%s_%s"), *(SerializedType.ToString()), *(Name));
 
-		if (!ManifestSerializer.OpenForWrite(FileName, Header, Error))
-		{
-			UE_LOG(ManifestSerialization, Warning, TEXT("Error Opening Sequence Sequencer File: Subject '%s' Error '%s'"), *(Name), *(Error.ToString()));
+			if (!ManifestSerializer.OpenForWrite(FileName, Header, Error))
+			{
+				UE_LOG(ManifestSerialization, Warning, TEXT("Error Opening Sequence Sequencer File: Subject '%s' Error '%s'"), *(Name), *(Error.ToString()));
+			}
 		}
 	}
 
diff --git a/Engine/Plugins/VirtualProduction/Takes/Source/TakeRecorder/Private/Widgets/STakeRecorderCockpit.cpp b/Engine/Plugins/VirtualProduction/Takes/Source/TakeRecorder/Private/Widgets/STakeRecorderCockpit.cpp
--- a/Engine/Plugins/VirtualProduction/Takes/Source/TakeRecorder/Private/Widgets/STakeRecorderCockpit.cpp
+++ b/Engine/Plugins/VirtualProduction/Takes/Source/TakeRecorder/Private/Widgets/STakeRecorderCockpit.cpp
@@ -151,6 +151,12 @@ namespace UE::TakeRecorder::Private
 	const FName TakeSaveDirPropertyName(GET_MEMBER_NAME_CHECKED(FTakeRecorderProjectParameters, TakeSaveDir));
 	/** Property name of the RootTakeSaveDir. */
 	const FName RootTakeSaveDirPropertyName(GET_MEMBER_NAME_CHECKED(FTakeRecorderProjectParameters, RootTakeSaveDir));
+
+	/** Verify the token key is allowed to be user defined. */
+	bool CanTokenBeUserDefined(const FString& InTokenKey)
+	{
+		return !InTokenKey.Equals(TokenKeySlate) && !InTokenKey.Equals(TokenKeyTake);
+	}
 }
 
 void STakeRecorderCockpit::Construct(const FArguments& InArgs)
@@ -544,7 +550,10 @@ void STakeRecorderCockpit::UpdateRecordError()
 
 	if (TakeMetaData)
 	{
-		PackageName = TakeMetaData->GenerateAssetPath(GetDefault<UTakeRecorderProjectSettings>()->Settings.GetTakeAssetPath());
+		if (!TakeMetaData->TryGenerateRootAssetPath(GetDefault<UTakeRecorderProjectSettings>()->Settings.GetTakeAssetPath(), PackageName, &RecordErrorText))
+		{
+			return;
+		}
 		FText OutReason;
 		if (!FPackageName::IsValidLongPackageName(PackageName, false, &OutReason))
 		{
@@ -1156,12 +1165,15 @@ FText STakeRecorderCockpit::GetCustomTokenTextValue(FString InTokenKey) const
 
 void STakeRecorderCockpit::SetCustomTokenTextValue(const FText& InNewText, ETextCommit::Type InCommitType, FString InTokenKey)
 {
-	FScopedTransaction Transaction(LOCTEXT("SetCustomTokenTextValue_Transaction", "Set Custom Token Text"));
-	const FNamingTokenData CustomToken(InTokenKey);
-	GetNamingTokensData()->Modify();
-	GetNamingTokensData()->UserDefinedTokens.FindOrAdd(CustomToken) = InNewText;
-	GetNamingTokensData()->VisibleUserTokens.Add(CustomToken.TokenKey);
-	OnTokenValueUpdated();
+		FScopedTransaction Transaction(LOCTEXT("SetCustomTokenTextValue_Transaction", "Set Custom Token Text"));
+		const FNamingTokenData CustomToken(InTokenKey);
+		GetNamingTokensData()->Modify();
+		if (UE::TakeRecorder::Private::CanTokenBeUserDefined(InTokenKey))
+		{
+			GetNamingTokensData()->UserDefinedTokens.FindOrAdd(CustomToken) = InNewText;
+		}
+		GetNamingTokensData()->VisibleUserTokens.Add(CustomToken.TokenKey);
+		OnTokenValueUpdated();
 }
 
 void STakeRecorderCockpit::OnTokenValueUpdated()
@@ -1438,10 +1450,12 @@ void STakeRecorderCockpit::CreateUserTokensUI()
 			for (const FString& UndefinedTokenStr : UndefinedTokens.UndefinedKeys)
 			{
 				FNamingTokenData UndefinedToken(UndefinedTokenStr);
-				if (!GetNamingTokensData()->UserDefinedTokens.Contains(UndefinedToken))
+				TMap<FNamingTokenData, FText>& UserDefinedTokens = GetNamingTokensData()->UserDefinedTokens;
+				if (!UserDefinedTokens.Contains(UndefinedToken)
+					&& UE::TakeRecorder::Private::CanTokenBeUserDefined(UndefinedTokenStr))
 				{
 					// Check for existence so we don't overwrite the user defined value.
-					GetNamingTokensData()->UserDefinedTokens.Add(UndefinedToken);
+					UserDefinedTokens.Add(UndefinedToken);
 				}
 				GetNamingTokensData()->VisibleUserTokens.Add(UndefinedToken.TokenKey);
 			}
diff --git a/Engine/Plugins/VirtualProduction/Takes/Source/TakeRecorder/Private/Widgets/STakeRecorderPanel.cpp b/Engine/Plugins/VirtualProduction/Takes/Source/TakeRecorder/Private/Widgets/STakeRecorderPanel.cpp
--- a/Engine/Plugins/VirtualProduction/Takes/Source/TakeRecorder/Private/Widgets/STakeRecorderPanel.cpp
+++ b/Engine/Plugins/VirtualProduction/Takes/Source/TakeRecorder/Private/Widgets/STakeRecorderPanel.cpp
@@ -966,24 +966,27 @@ void STakeRecorderPanel::ToggleTakeBrowserCheckState(ECheckBoxState CheckState)
 		FAssetRegistryModule& AssetRegistryModule = FModuleManager::LoadModuleChecked<FAssetRegistryModule>(AssetRegistryConstants::ModuleName);
 		IAssetRegistry& AssetRegistry = AssetRegistryModule.Get();
 
-		FString TakesPath = GetTakeMetaData()->GenerateAssetPath(GetDefault<UTakeRecorderProjectSettings>()->Settings.GetTakeAssetPath());
-		TakesPath = FPaths::GetPath(*TakesPath);
-
-		while(!TakesPath.IsEmpty())
+		FString TakesPath;
+		if (GetTakeMetaData()->TryGenerateRootAssetPath(GetDefault<UTakeRecorderProjectSettings>()->Settings.GetTakeAssetPath(), TakesPath))
 		{
-			if (AssetRegistry.HasAssets(FName(*TakesPath), true))
+			TakesPath = FPaths::GetPath(*TakesPath);
+
+			while(!TakesPath.IsEmpty())
 			{
-				break;
+				if (AssetRegistry.HasAssets(FName(*TakesPath), true))
+				{
+					break;
+				}
+				TakesPath = FPaths::GetPath(TakesPath);
 			}
-			TakesPath = FPaths::GetPath(TakesPath);
-		}
 
-		TArray<FString> TakesFolder;
-		TakesFolder.Push(TakesPath);
-		if (AssetRegistry.HasAssets(FName(*TakesPath), true) )
-		{
-			FContentBrowserModule& ContentBrowserModule = FModuleManager::Get().LoadModuleChecked<FContentBrowserModule>("ContentBrowser");
-			ContentBrowserModule.Get().SyncBrowserToFolders(TakesFolder, bAllowLockedBrowser, bFocusContentBrowser, ITakeRecorderModule::TakesBrowserInstanceName );
+			TArray<FString> TakesFolder;
+			TakesFolder.Push(TakesPath);
+			if (AssetRegistry.HasAssets(FName(*TakesPath), true) )
+			{
+				FContentBrowserModule& ContentBrowserModule = FModuleManager::Get().LoadModuleChecked<FContentBrowserModule>("ContentBrowser");
+				ContentBrowserModule.Get().SyncBrowserToFolders(TakesFolder, bAllowLockedBrowser, bFocusContentBrowser, ITakeRecorderModule::TakesBrowserInstanceName );
+			}
 		}
 
 		TakesBrowserTab->FlashTab();
diff --git a/Engine/Plugins/VirtualProduction/Takes/Source/TakeRecorderNamingTokens/Private/TakeRecorderNamingTokens.cpp b/Engine/Plugins/VirtualProduction/Takes/Source/TakeRecorderNamingTokens/Private/TakeRecorderNamingTokens.cpp
--- a/Engine/Plugins/VirtualProduction/Takes/Source/TakeRecorderNamingTokens/Private/TakeRecorderNamingTokens.cpp
+++ b/Engine/Plugins/VirtualProduction/Takes/Source/TakeRecorderNamingTokens/Private/TakeRecorderNamingTokens.cpp
@@ -2,16 +2,22 @@
 
 #include "TakeRecorderNamingTokens.h"
 
-#include "Editor.h"
-#include "ITakeRecorderNamingTokensModule.h"
-#include "Misc/PackageName.h"
-#include "MovieSceneToolsProjectSettings.h"
 #include "NamingTokens/TakeRecorderNamingTokensContext.h"
-#include "TakeMetaData.h"
+
+#include "ITakeRecorderNamingTokensModule.h"
 #include "TakeRecorderNamingTokensLog.h"
+
+#include "Recorder/TakeRecorderSubsystem.h"
+#include "TakeMetaData.h"
 #include "TakeRecorderSettings.h"
 #include "TakesCoreBlueprintLibrary.h"
 
+#include "Editor.h"
+#include "Engine/Engine.h"
+#include "LevelSequence.h"
+#include "Misc/PackageName.h"
+#include "MovieSceneToolsProjectSettings.h"
+
 #define LOCTEXT_NAMESPACE "TakeRecorderNamingTokens"
 
 UTakeRecorderNamingTokens::UTakeRecorderNamingTokens()
@@ -77,7 +83,7 @@ void UTakeRecorderNamingTokens::OnCreateDefaultTokens(TArray<FNamingTokenData>&
 
 	Tokens.Add({
 	    TEXT("take"),
-	    LOCTEXT("TokenTake", "Take"),
+	    LOCTEXT("TokenTake", "Take Number"),
 	    FNamingTokenData::FTokenProcessorDelegateNative::CreateLambda([this]{
 	        const UMovieSceneToolsProjectSettings* ProjectSettings = GetDefault<UMovieSceneToolsProjectSettings>();
 	        const int32 TakeNumDigits = ProjectSettings->TakeNumDigits;
@@ -93,6 +99,18 @@ void UTakeRecorderNamingTokens::OnCreateDefaultTokens(TArray<FNamingTokenData>&
 	    })
 	});
 
+	Tokens.Add({
+		TEXT("takeName"),
+		LOCTEXT("TokenTakeName", "Take Name"),
+		FNamingTokenData::FTokenProcessorDelegateNative::CreateLambda([this]{
+			if (const ULevelSequence* Sequence = GEngine->GetEngineSubsystem<UTakeRecorderSubsystem>()->GetLevelSequence())
+			{
+				return FText::FromString(Sequence->GetName());
+			}
+			return FText::GetEmpty();
+		})
+	});
+
 	Tokens.Add({
 	    TEXT("map"),
 	    LOCTEXT("TokenMap", "Map"),
diff --git a/Engine/Plugins/VirtualProduction/Takes/Source/TakeRecorderSources/Private/TakeRecorderMicrophoneAudioSource.cpp b/Engine/Plugins/VirtualProduction/Takes/Source/TakeRecorderSources/Private/TakeRecorderMicrophoneAudioSource.cpp
--- a/Engine/Plugins/VirtualProduction/Takes/Source/TakeRecorderSources/Private/TakeRecorderMicrophoneAudioSource.cpp
+++ b/Engine/Plugins/VirtualProduction/Takes/Source/TakeRecorderSources/Private/TakeRecorderMicrophoneAudioSource.cpp
@@ -34,7 +34,7 @@ UTakeRecorderMicrophoneAudioSourceSettings::UTakeRecorderMicrophoneAudioSourceSe
 	: Super(ObjInit)
 	, AudioSourceName(NSLOCTEXT("UTakeRecorderMicrophoneAudioSource", "Label", "Microphone Audio"))
 	, AudioTrackName(NSLOCTEXT("UTakeRecorderMicrophoneAudioSource", "DefaultAudioTrackName", "Recorded Audio"))
-	, AudioAssetName(TEXT("Audio_{slate}_{take}_Channel_{channel}"))
+	, AudioAssetName(TEXT("Audio_{takeName}_Channel_{channel}"))
 	, AudioSubDirectory(TEXT("Audio"))
 {
 	TrackTint = FColor(75, 67, 148);
diff --git a/Engine/Plugins/VirtualProduction/Takes/Source/TakeTrackRecorders/Public/TrackRecorders/MovieSceneAnimationTrackRecorderSettings.h b/Engine/Plugins/VirtualProduction/Takes/Source/TakeTrackRecorders/Public/TrackRecorders/MovieSceneAnimationTrackRecorderSettings.h
--- a/Engine/Plugins/VirtualProduction/Takes/Source/TakeTrackRecorders/Public/TrackRecorders/MovieSceneAnimationTrackRecorderSettings.h
+++ b/Engine/Plugins/VirtualProduction/Takes/Source/TakeTrackRecorders/Public/TrackRecorders/MovieSceneAnimationTrackRecorderSettings.h
@@ -17,7 +17,7 @@ public:
 	UMovieSceneAnimationTrackRecorderEditorSettings(const FObjectInitializer& ObjInit)
 		: Super(ObjInit)
 	, AnimationTrackName(NSLOCTEXT("UMovieSceneAnimationTrackRecorderSettings", "DefaultAnimationTrackName", "RecordedAnimation"))
-	, AnimationAssetName(TEXT("{actor}_{slate}_{take}"))
+	, AnimationAssetName(TEXT("{actor}_{takeName}"))
 	, AnimationSubDirectory(TEXT("Animation"))
 	, bRemoveRootAnimation(true)
 	{
diff --git a/Engine/Plugins/VirtualProduction/Takes/Source/TakesCore/Private/TakeMetaData.cpp b/Engine/Plugins/VirtualProduction/Takes/Source/TakesCore/Private/TakeMetaData.cpp
--- a/Engine/Plugins/VirtualProduction/Takes/Source/TakesCore/Private/TakeMetaData.cpp
+++ b/Engine/Plugins/VirtualProduction/Takes/Source/TakesCore/Private/TakeMetaData.cpp
@@ -14,6 +14,7 @@
 #include "TakesCoreLog.h"
 #include "UObject/AssetRegistryTagsContext.h"
 #include "UObject/Package.h"
+#include "Utils/NamingTokenUtils.h"
 
 #include UE_INLINE_GENERATED_CPP_BY_NAME(TakeMetaData)
 
@@ -131,6 +132,26 @@ FString UTakeMetaData::GenerateAssetPath(const FString& PathFormatString, UTakeR
 	return Result.EvaluatedText.ToString();
 }
 
+bool UTakeMetaData::TryGenerateRootAssetPath(const FString& PathFormatString, FString& OutGeneratedAssetPath,
+	FText* OutErrorMessage, UTakeRecorderNamingTokensContext* InContext) const
+{
+	const FString TakeNameToken = TEXT("takeName");
+	if (UE::NamingTokens::Utils::IsTokenInString(TakeNameToken, PathFormatString))
+	{
+		if (OutErrorMessage)
+		{
+			*OutErrorMessage =
+				FText::Format(NSLOCTEXT("TakeMetaData", "ErrorGenerateRootAssetPath_InvalidToken",
+					"Token '{0}' cannot be present in this asset path {1}."),
+					FText::FromString(TakeNameToken), FText::FromString(PathFormatString));
+		}
+		return false;
+	}
+
+	OutGeneratedAssetPath = GenerateAssetPath(PathFormatString, InContext);
+	return true;
+}
+
 FNamingTokenResultData UTakeMetaData::ProcessTokens(const FText& InText, UTakeRecorderNamingTokensContext* InContext) const
 {
 	if (!ensure(GEngine))
diff --git a/Engine/Plugins/VirtualProduction/Takes/Source/TakesCore/Public/TakeMetaData.h b/Engine/Plugins/VirtualProduction/Takes/Source/TakesCore/Public/TakeMetaData.h
--- a/Engine/Plugins/VirtualProduction/Takes/Source/TakesCore/Public/TakeMetaData.h
+++ b/Engine/Plugins/VirtualProduction/Takes/Source/TakesCore/Public/TakeMetaData.h
@@ -191,6 +191,12 @@ public:
 	 */
 	UFUNCTION(BlueprintCallable, Category="Take")
 	FString GenerateAssetPath(const FString& PathFormatString, UTakeRecorderNamingTokensContext* InContext = nullptr) const;
+
+	/**
+	 * Attempt to generate the root asset path.
+	 * @return true if successful, false if the path couldn't be generated.
+	 */
+	bool TryGenerateRootAssetPath(const FString& PathFormatString, FString& OutGeneratedAssetPath, FText* OutErrorMessage = nullptr, UTakeRecorderNamingTokensContext* InContext = nullptr) const;
 	
 	/** Process TakeRecorder Naming Tokens. */
 	FNamingTokenResultData ProcessTokens(const FText& InText, UTakeRecorderNamingTokensContext* InContext = nullptr) const;
