TakeRecorder: Change subsystem from editor tickable to a game tickable object, with optimizations to prevent tickable registration until specifically used. Allow SetTargetSequence to be called in BP without passing any parameters.

diff --git a/Engine/Plugins/VirtualProduction/Takes/Source/TakeRecorder/Private/Recorder/TakeRecorderSubsystem.cpp b/Engine/Plugins/VirtualProduction/Takes/Source/TakeRecorder/Private/Recorder/TakeRecorderSubsystem.cpp
--- a/Engine/Plugins/VirtualProduction/Takes/Source/TakeRecorder/Private/Recorder/TakeRecorderSubsystem.cpp
+++ b/Engine/Plugins/VirtualProduction/Takes/Source/TakeRecorder/Private/Recorder/TakeRecorderSubsystem.cpp
@@ -71,6 +71,35 @@ void UTakeRecorderSubsystem::Deinitialize()
 	}
 
 	UnbindNamingTokensEvents();
+
+	bHasTargetSequenceBeenSet = false;
+	SetTickableTickType(ETickableTickType::Never);
+}
+
+ETickableTickType UTakeRecorderSubsystem::GetTickableTickType() const
+{
+	// This is to prevent registration until we call SetTargetSequence.
+	return ETickableTickType::Never;
+}
+
+UWorld* UTakeRecorderSubsystem::GetTickableGameObjectWorld() const
+{
+	return FTickableGameObject::GetTickableGameObjectWorld();
+}
+
+bool UTakeRecorderSubsystem::IsTickable() const
+{
+	return bHasTargetSequenceBeenSet;
+}
+
+bool UTakeRecorderSubsystem::IsTickableWhenPaused() const
+{
+	return FTickableGameObject::IsTickableWhenPaused();
+}
+
+bool UTakeRecorderSubsystem::IsTickableInEditor() const
+{
+	return true;
 }
 
 void UTakeRecorderSubsystem::Tick(float DeltaTime)
@@ -168,6 +197,9 @@ void UTakeRecorderSubsystem::SetTargetSequence(const FTakeRecorderSequenceParame
 
 	/** Clear the dirty flag since the preset was just initialized. */
 	TransientPreset->GetOutermost()->SetDirtyFlag(false);
+
+	bHasTargetSequenceBeenSet = true;
+	SetTickableTickType(ETickableTickType::Conditional);
 }
 
 void UTakeRecorderSubsystem::SetRecordIntoLevelSequence(ULevelSequence* LevelSequence)
diff --git a/Engine/Plugins/VirtualProduction/Takes/Source/TakeRecorder/Public/Recorder/TakeRecorderSubsystem.h b/Engine/Plugins/VirtualProduction/Takes/Source/TakeRecorder/Public/Recorder/TakeRecorderSubsystem.h
--- a/Engine/Plugins/VirtualProduction/Takes/Source/TakeRecorder/Public/Recorder/TakeRecorderSubsystem.h
+++ b/Engine/Plugins/VirtualProduction/Takes/Source/TakeRecorder/Public/Recorder/TakeRecorderSubsystem.h
@@ -5,7 +5,7 @@
 #include "Subsystems/EngineSubsystem.h"
 #include "TakeRecorder.h"
 #include "TakeRecorderSources.h"
-#include "TickableEditorObject.h"
+#include "Tickable.h"
 
 #include "TakeRecorderSubsystem.generated.h"
 
@@ -30,11 +30,12 @@ DECLARE_DYNAMIC_MULTICAST_DELEGATE_OneParam(FOnTakeRecorderSourceAddedDynamic, U
 DECLARE_DYNAMIC_MULTICAST_DELEGATE_OneParam(FOnTakeRecorderSourceRemovedDynamic, UTakeRecorderSource*, Source);
 
 /**
-* UTakeRecorderSubsystem
-* Subsystem for Take Recorder
+* UTakeRecorderSubsystem Subsystem for Take Recorder.
+* SetTargetSequence() needs to be called at least once prior to use. This will perform additional initialization
+* and register the subsystem as tickable currently required to cache level metadata.
 */
 UCLASS()
-class TAKERECORDER_API UTakeRecorderSubsystem : public UEngineSubsystem, public FTickableEditorObject
+class TAKERECORDER_API UTakeRecorderSubsystem : public UEngineSubsystem, public FTickableGameObject
 {
 	GENERATED_BODY()
 
@@ -43,14 +44,24 @@ public:
 	virtual void Initialize(FSubsystemCollectionBase& Collection) override;
 	virtual void Deinitialize() override;
 
-	// Begin FTickableEditorObject
+	// Begin FTickableGameObject
+	virtual ETickableTickType GetTickableTickType() const override;
+	virtual UWorld* GetTickableGameObjectWorld() const override;
+	virtual bool IsTickable() const override;
+	virtual bool IsTickableWhenPaused() const override;
+	virtual bool IsTickableInEditor() const override;
 	virtual void Tick(float DeltaTime) override;
 	virtual TStatId GetStatId() const override;
-	// ~End FTickableEditorObject
+	// ~End FTickableGameObject
 	
-	/** Provide sequence data for this take recorder. */
-	UFUNCTION(BlueprintCallable, Category = TakeRecorder)
-	void SetTargetSequence(const FTakeRecorderSequenceParameters& InData);
+	/**
+	 * Provide sequence data for this take recorder. This will also perform initialization of the subsystem.
+	 * This must be called prior to any usage.
+	 *
+	 * @param InData The sequence parameters, which are mutually exclusive. They can all be null.
+	 */
+	UFUNCTION(BlueprintCallable, Category = TakeRecorder, meta=(AutoCreateRefTerm="InData"))
+	void SetTargetSequence(const FTakeRecorderSequenceParameters& InData = FTakeRecorderSequenceParameters());
 	
 	/** Set the record into level sequence. */
 	void SetRecordIntoLevelSequence(ULevelSequence* LevelSequence);
@@ -420,6 +431,9 @@ private:
 	
 	/** Whether we should auto apply the next available take number when asset discovery has finished or not */
 	bool bAutoApplyTakeNumber;
+
+	/** If SetTargetSequence has been called, and we're fully initialized. */
+	bool bHasTargetSequenceBeenSet = false;
 	
 public:
 
