TakeRecorder: Expose GetLastRecordedLevelSequence in the subsystem to BP.

diff --git a/Engine/Plugins/VirtualProduction/Takes/Source/TakeRecorder/Public/Recorder/TakeRecorderSubsystem.h b/Engine/Plugins/VirtualProduction/Takes/Source/TakeRecorder/Public/Recorder/TakeRecorderSubsystem.h
--- a/Engine/Plugins/VirtualProduction/Takes/Source/TakeRecorder/Public/Recorder/TakeRecorderSubsystem.h
+++ b/Engine/Plugins/VirtualProduction/Takes/Source/TakeRecorder/Public/Recorder/TakeRecorderSubsystem.h
@@ -276,6 +276,7 @@ public:
 	ULevelSequence* GetRecordIntoLevelSequence() const;
 
 	/** Retrieve the Last Recorded Level Sequence if it exists. */
+	UFUNCTION(BlueprintCallable, Category = TakeRecorder)
 	ULevelSequence* GetLastRecordedLevelSequence() const;
 
 	/** Retrieve the transient preset the subsystem is using. */
