TakeRecorder: Take meta data now tracks if it was recorded as a subsequence. Add IsSubSequence and SetIsSubSequence to take meta data.

diff --git a/Engine/Plugins/VirtualProduction/Takes/Source/TakesCore/Private/TakeMetaData.cpp b/Engine/Plugins/VirtualProduction/Takes/Source/TakesCore/Private/TakeMetaData.cpp
--- a/Engine/Plugins/VirtualProduction/Takes/Source/TakesCore/Private/TakeMetaData.cpp
+++ b/Engine/Plugins/VirtualProduction/Takes/Source/TakesCore/Private/TakeMetaData.cpp
@@ -343,6 +343,19 @@ void UTakeMetaData::SetFrameRateFromTimecode(bool InFromTimecode)
 	}
 }
 
+bool UTakeMetaData::IsSubSequence() const
+{
+	return bIsSubSequence;
+}
+
+void UTakeMetaData::SetIsSubSequence(bool bNewValue)
+{
+	if (!bIsLocked)
+	{
+		bIsSubSequence = bNewValue;
+	}
+}
+
 
 void UTakeMetaData::ExtendAssetRegistryTags(FAssetRegistryTagsContext Context) const
 {
diff --git a/Engine/Plugins/VirtualProduction/Takes/Source/TakesCore/Private/TakeRecorderSources.cpp b/Engine/Plugins/VirtualProduction/Takes/Source/TakesCore/Private/TakeRecorderSources.cpp
--- a/Engine/Plugins/VirtualProduction/Takes/Source/TakesCore/Private/TakeRecorderSources.cpp
+++ b/Engine/Plugins/VirtualProduction/Takes/Source/TakesCore/Private/TakeRecorderSources.cpp
@@ -479,6 +479,7 @@ void UTakeRecorderSources::StopRecording(class ULevelSequence* InSequence, const
 				ULevelSequence* SequenceAsset = CastChecked<ULevelSequence>(SubSequence);
 				UTakeMetaData* AssetMetaData = SequenceAsset->FindMetaData<UTakeMetaData>();
 				check(AssetMetaData);
+				AssetMetaData->SetIsSubSequence(true);
 				AssetMetaData->Lock();
 
 				SubSection->SetRange(SubSequence->GetMovieScene()->GetPlaybackRange());
diff --git a/Engine/Plugins/VirtualProduction/Takes/Source/TakesCore/Public/TakeMetaData.h b/Engine/Plugins/VirtualProduction/Takes/Source/TakesCore/Public/TakeMetaData.h
--- a/Engine/Plugins/VirtualProduction/Takes/Source/TakesCore/Public/TakeMetaData.h
+++ b/Engine/Plugins/VirtualProduction/Takes/Source/TakesCore/Public/TakeMetaData.h
@@ -278,6 +278,16 @@ public:
 	UFUNCTION(BlueprintCallable, Category = "Take")
 	void SetFrameRateFromTimecode(bool InFromTimecode);
 
+	/**
+	 * True if this take meta data was recorded as part of a subsequence.
+	 */
+	UFUNCTION(BlueprintCallable, Category = "Take")
+	bool IsSubSequence() const;
+
+	/** Set whether the take meta data is part of a subsequence. */
+	UFUNCTION(BlueprintCallable, Category = "Take")
+	void SetIsSubSequence(bool bNewValue);
+
 private:
 
 	/** Whether the take is locked */
@@ -329,4 +339,8 @@ private:
 	/** Whether or not we get or frame rate from Timecode, default to true */
 	UPROPERTY()
 	bool bFrameRateFromTimecode;
+
+	/** If this meta data was recorded as a subsequence. */
+	UPROPERTY()
+	bool bIsSubSequence = false;
 };
