TakeRecorder: The takeName token will now attempt to calculate the most correct result if it is read when only the transient preset is available.

diff --git a/Engine/Plugins/VirtualProduction/Takes/Source/TakeRecorderNamingTokens/Private/TakeRecorderNamingTokens.cpp b/Engine/Plugins/VirtualProduction/Takes/Source/TakeRecorderNamingTokens/Private/TakeRecorderNamingTokens.cpp
--- a/Engine/Plugins/VirtualProduction/Takes/Source/TakeRecorderNamingTokens/Private/TakeRecorderNamingTokens.cpp
+++ b/Engine/Plugins/VirtualProduction/Takes/Source/TakeRecorderNamingTokens/Private/TakeRecorderNamingTokens.cpp
@@ -16,6 +16,7 @@
 #include "Engine/Engine.h"
 #include "LevelSequence.h"
 #include "Misc/PackageName.h"
+#include "Misc/Paths.h"
 #include "MovieSceneToolsProjectSettings.h"
 
 #define LOCTEXT_NAMESPACE "TakeRecorderNamingTokens"
@@ -102,10 +103,33 @@ void UTakeRecorderNamingTokens::OnCreateDefaultTokens(TArray<FNamingTokenData>&
 	Tokens.Add({
 		TEXT("takeName"),
 		LOCTEXT("TokenTakeName", "Take Name"),
+		LOCTEXT("TokenTakeNameDescription",
+			"The sequencer asset's name. This is only completely accurate when read at the time of recording or the subsystem has a valid, non-transient, level sequence."),
 		FNamingTokenData::FTokenProcessorDelegateNative::CreateLambda([this]{
-			if (const ULevelSequence* Sequence = GEngine->GetEngineSubsystem<UTakeRecorderSubsystem>()->GetLevelSequence())
+			if (const UTakeRecorderSubsystem* Subsystem = GEngine->GetEngineSubsystem<UTakeRecorderSubsystem>())
 			{
-				return FText::FromString(Sequence->GetName());
+				if (const ULevelSequence* Sequence = Subsystem->GetLevelSequence())
+				{
+					FString AssetName = Sequence->GetName();
+					
+					const UTakePreset* TransientPreset = Subsystem->GetTransientPreset();
+					if (TransientPreset && Sequence == TransientPreset->GetLevelSequence() && TakeMetaData)
+					{
+						// Working with the transient preset... try to generate the most correct takeName.
+						// Depending on when this is read, it may not match the actual takeName (Sequencer asset's name)
+						// when saving the asset.
+						
+						const FString TakeAssetPath = GetDefault<UTakeRecorderProjectSettings>()->Settings.GetTakeAssetPath();
+						
+						FString GeneratedPath;
+						if (TakeMetaData->TryGenerateRootAssetPath(TakeAssetPath, GeneratedPath))
+						{
+							AssetName = FPaths::GetBaseFilename(GeneratedPath);
+						}
+					}
+
+					return FText::FromString(AssetName);
+				}
 			}
 			return FText::GetEmpty();
 		})
