TakeRecorder: Extend the IsSubSequence flag to be asset registry searchable.

diff --git a/Engine/Plugins/VirtualProduction/Takes/Source/TakesCore/Private/TakeMetaData.cpp b/Engine/Plugins/VirtualProduction/Takes/Source/TakesCore/Private/TakeMetaData.cpp
--- a/Engine/Plugins/VirtualProduction/Takes/Source/TakesCore/Private/TakeMetaData.cpp
+++ b/Engine/Plugins/VirtualProduction/Takes/Source/TakesCore/Private/TakeMetaData.cpp
@@ -25,6 +25,7 @@ const FName UTakeMetaData::AssetRegistryTag_TimecodeIn = "TakeMetaData_TimecodeI
 const FName UTakeMetaData::AssetRegistryTag_TimecodeOut = "TakeMetaData_TimecodeOut";
 const FName UTakeMetaData::AssetRegistryTag_Description = "TakeMetaData_Description";
 const FName UTakeMetaData::AssetRegistryTag_LevelPath   = "TakeMetaData_LevelPath";
+const FName UTakeMetaData::AssetRegistryTag_IsSubSequence = "TakeMetaData_IsSubSequence";
 
 UTakeMetaData::UTakeMetaData(const FObjectInitializer& ObjInit)
 	: Super(ObjInit)
@@ -373,6 +374,8 @@ void UTakeMetaData::ExtendAssetRegistryTags(FAssetRegistryTagsContext Context) c
 										FAssetRegistryTag::ETagType::TT_Numerical,		FAssetRegistryTag::TD_None));
 	Context.AddTag(FAssetRegistryTag(	AssetRegistryTag_Description,					Description,
 										FAssetRegistryTag::ETagType::TT_Alphabetical,	FAssetRegistryTag::TD_None));
+	Context.AddTag(FAssetRegistryTag(	AssetRegistryTag_IsSubSequence,					IsSubSequence() ? TEXT("True") : TEXT("False"),
+										FAssetRegistryTag::ETagType::TT_Alphabetical,	FAssetRegistryTag::TD_None));
 }
 
 void UTakeMetaData::ExtendAssetRegistryTagMetaData(TMap<FName, FAssetRegistryTagMetadata>& OutMetadata) const
@@ -411,6 +414,11 @@ void UTakeMetaData::ExtendAssetRegistryTagMetaData(TMap<FName, FAssetRegistryTag
 		.SetDisplayName(NSLOCTEXT("TakeMetaData", "LevelPath_Label", "Map"))
 		.SetTooltip(    NSLOCTEXT("TakeMetaData", "LevelPath_Tip",   "Map used for this take"))
 	);
+
+	OutMetadata.Add(AssetRegistryTag_IsSubSequence, FAssetRegistryTagMetadata()
+		.SetDisplayName(NSLOCTEXT("TakeMetaData", "IsSubSequence_Label", "Is SubSequence"))
+		.SetTooltip(    NSLOCTEXT("TakeMetaData", "IsSubSequence_Tip",   "If this take was recorded as a subsequence"))
+	);
 }
 
 
diff --git a/Engine/Plugins/VirtualProduction/Takes/Source/TakesCore/Public/TakeMetaData.h b/Engine/Plugins/VirtualProduction/Takes/Source/TakesCore/Public/TakeMetaData.h
--- a/Engine/Plugins/VirtualProduction/Takes/Source/TakesCore/Public/TakeMetaData.h
+++ b/Engine/Plugins/VirtualProduction/Takes/Source/TakesCore/Public/TakeMetaData.h
@@ -47,7 +47,10 @@ public:
 
 	/** The asset registry tag that contains the level-path for this meta-data */
 	static const FName AssetRegistryTag_LevelPath;
-
+	
+	/** The asset registry tag that contains whether this is a subsequence. */
+	static const FName AssetRegistryTag_IsSubSequence;
+	
 	/**
 	 * Access the global config instance that houses default settings for take meta data for a given project
 	 */
diff --git a/Engine/Plugins/VirtualProduction/VirtualCamera/Source/VirtualCamera/Private/FunctionLibraries/TakeMetaDataTagsFunctionLibrary.cpp b/Engine/Plugins/VirtualProduction/VirtualCamera/Source/VirtualCamera/Private/FunctionLibraries/TakeMetaDataTagsFunctionLibrary.cpp
--- a/Engine/Plugins/VirtualProduction/VirtualCamera/Source/VirtualCamera/Private/FunctionLibraries/TakeMetaDataTagsFunctionLibrary.cpp
+++ b/Engine/Plugins/VirtualProduction/VirtualCamera/Source/VirtualCamera/Private/FunctionLibraries/TakeMetaDataTagsFunctionLibrary.cpp
@@ -69,6 +69,15 @@ FName UTakeMetaDataTagsFunctionLibrary::GetTakeMetaDataTag_LevelPath()
 #endif
 }
 
+FName UTakeMetaDataTagsFunctionLibrary::GetTakeMetaDataTag_IsSubSequence()
+{
+#if WITH_EDITOR
+	return UTakeMetaData::AssetRegistryTag_IsSubSequence;
+#else
+	return "TakeMetaData_IsSubSequence";
+#endif
+}
+
 TSet<FName> UTakeMetaDataTagsFunctionLibrary::GetAllTakeMetaDataTags()
 {
 	return {
@@ -78,6 +87,7 @@ TSet<FName> UTakeMetaDataTagsFunctionLibrary::GetAllTakeMetaDataTags()
 		GetTakeMetaDataTag_TimecodeIn(),
 		GetTakeMetaDataTag_TimecodeOut(),
 		GetTakeMetaDataTag_Description(),
-		GetTakeMetaDataTag_LevelPath()
+		GetTakeMetaDataTag_LevelPath(),
+		GetTakeMetaDataTag_IsSubSequence()
 	};
 }
diff --git a/Engine/Plugins/VirtualProduction/VirtualCamera/Source/VirtualCamera/Public/FunctionLibraries/TakeMetaDataTagsFunctionLibrary.h b/Engine/Plugins/VirtualProduction/VirtualCamera/Source/VirtualCamera/Public/FunctionLibraries/TakeMetaDataTagsFunctionLibrary.h
--- a/Engine/Plugins/VirtualProduction/VirtualCamera/Source/VirtualCamera/Public/FunctionLibraries/TakeMetaDataTagsFunctionLibrary.h
+++ b/Engine/Plugins/VirtualProduction/VirtualCamera/Source/VirtualCamera/Public/FunctionLibraries/TakeMetaDataTagsFunctionLibrary.h
@@ -40,6 +40,10 @@ public:
 	UFUNCTION(BlueprintPure, Category = "VirtualCamera")
 	static FName GetTakeMetaDataTag_LevelPath();
 
+	/** The asset registry tag that specifies if the meta-data is a subsequence. */
+	UFUNCTION(BlueprintPure, Category = "VirtualCamera")
+	static FName GetTakeMetaDataTag_IsSubSequence();
+
 	/** Gets all tags for TakeMetaData */
 	UFUNCTION(BlueprintPure, Category = "VirtualCamera")
 	static TSet<FName> GetAllTakeMetaDataTags();
