TakeRecorder: Add missing multicast delegates, TakeRecorderSlateChanged and TakeRecorderTakeNumberChanged.

diff --git a/Engine/Plugins/VirtualProduction/Takes/Source/TakeRecorder/Private/Recorder/TakeRecorderSubsystem.cpp b/Engine/Plugins/VirtualProduction/Takes/Source/TakeRecorder/Private/Recorder/TakeRecorderSubsystem.cpp
--- a/Engine/Plugins/VirtualProduction/Takes/Source/TakeRecorder/Private/Recorder/TakeRecorderSubsystem.cpp
+++ b/Engine/Plugins/VirtualProduction/Takes/Source/TakeRecorder/Private/Recorder/TakeRecorderSubsystem.cpp
@@ -40,6 +40,9 @@ void UTakeRecorderSubsystem::Initialize(FSubsystemCollectionBase& Collection)
 	UTakeRecorderSources::OnSourceAdded().AddUObject(this, &UTakeRecorderSubsystem::OnSourceAdded);
 	UTakeRecorderSources::OnSourceRemoved().AddUObject(this, &UTakeRecorderSubsystem::OnSourceRemoved);
 
+	UTakeMetaData::OnTakeSlateChanged().AddUObject(this, &UTakeRecorderSubsystem::OnTakeSlateChanged);
+	UTakeMetaData::OnTakeNumberChanged().AddUObject(this, &UTakeRecorderSubsystem::OnTakeNumberChanged);
+
 	NamingTokensData = NewObject<UTakeRecorderNamingTokensData>();
 	NamingTokensData->SetFlags(RF_Transactional);
 
@@ -56,6 +59,9 @@ void UTakeRecorderSubsystem::Deinitialize()
 
 	UTakeRecorderSources::OnSourceAdded().RemoveAll(this);
 	UTakeRecorderSources::OnSourceRemoved().RemoveAll(this);
+
+	UTakeMetaData::OnTakeSlateChanged().RemoveAll(this);
+	UTakeMetaData::OnTakeNumberChanged().RemoveAll(this);
 	
 	if (const FAssetRegistryModule* AssetRegistryModule = FModuleManager::GetModulePtr<FAssetRegistryModule>("AssetRegistry"))
 	{
@@ -1021,6 +1027,16 @@ void UTakeRecorderSubsystem::OnRecordingCancelled(UTakeRecorder* Recorder)
 	TakeRecorderCancelled.Broadcast();
 }
 
+void UTakeRecorderSubsystem::OnTakeSlateChanged(const FString& InSlate, UTakeMetaData* InTakeMetaData)
+{
+	TakeRecorderSlateChanged.Broadcast(InSlate, InTakeMetaData);
+}
+
+void UTakeRecorderSubsystem::OnTakeNumberChanged(int32 InTakeNumber, UTakeMetaData* InTakeMetaData)
+{
+	TakeRecorderTakeNumberChanged.Broadcast(InTakeNumber, InTakeMetaData);
+}
+
 void UTakeRecorderSubsystem::SetLastLevelSequence(ULevelSequence* InSequence)
 {
 	LastRecordedLevelSequence = InSequence;
diff --git a/Engine/Plugins/VirtualProduction/Takes/Source/TakeRecorder/Public/Recorder/TakeRecorderSubsystem.h b/Engine/Plugins/VirtualProduction/Takes/Source/TakeRecorder/Public/Recorder/TakeRecorderSubsystem.h
--- a/Engine/Plugins/VirtualProduction/Takes/Source/TakeRecorder/Public/Recorder/TakeRecorderSubsystem.h
+++ b/Engine/Plugins/VirtualProduction/Takes/Source/TakeRecorder/Public/Recorder/TakeRecorderSubsystem.h
@@ -21,6 +21,9 @@ DECLARE_DYNAMIC_MULTICAST_DELEGATE_OneParam(FTakeRecorderFinished, ULevelSequenc
 DECLARE_DYNAMIC_MULTICAST_DELEGATE(FTakeRecorderCancelled);
 DECLARE_DYNAMIC_MULTICAST_DELEGATE_OneParam(FTakeRecorderMarkedFrameAdded, const FMovieSceneMarkedFrame&, MarkedFrame);
 
+DECLARE_DYNAMIC_MULTICAST_DELEGATE_TwoParams(FTakeRecorderSlateChanged, const FString&, Slate, UTakeMetaData*, TakeMetaData);
+DECLARE_DYNAMIC_MULTICAST_DELEGATE_TwoParams(FTakeRecorderTakeNumberChanged, int32, TakeNumber, UTakeMetaData*, TakeMetaData);
+
 /**
  * Dynamic delegates for source modification. These are reproductions of static ones declared under UTakeRecorderSource.
  * We do this so we can successfully leverage static delegates within the TakesCore module, firing when any change is made for
@@ -346,7 +349,9 @@ private:
 	void OnRecordingStopped(UTakeRecorder* Recorder);
 	void OnRecordingFinished(UTakeRecorder* Recorder);
 	void OnRecordingCancelled(UTakeRecorder* Recorder);
-
+	void OnTakeSlateChanged(const FString& InSlate, UTakeMetaData* InTakeMetaData);
+	void OnTakeNumberChanged(int32 InTakeNumber, UTakeMetaData* InTakeMetaData);
+	
 	/** Update the last level sequence. */
 	void SetLastLevelSequence(ULevelSequence *InSequence);
 
@@ -466,6 +471,14 @@ public:
 	UPROPERTY(BlueprintAssignable, Category=TakeRecorder)
 	FTakeRecorderMarkedFrameAdded TakeRecorderMarkedFrameAdded;
 
+	/** Called when a take recorder slate changes. */
+	UPROPERTY(BlueprintAssignable, Category=TakeRecorder)
+	FTakeRecorderSlateChanged TakeRecorderSlateChanged;
+
+	/** Called when a take recorder take number changes. */
+	UPROPERTY(BlueprintAssignable, Category=TakeRecorder)
+	FTakeRecorderTakeNumberChanged TakeRecorderTakeNumberChanged;
+	
 	/** Multicast delegate when any source has been added. */
 	UPROPERTY(BlueprintAssignable, Category=TakeRecorder)
 	FOnTakeRecorderSourceAddedDynamic TakeRecorderSourceAdded;
diff --git a/Engine/Plugins/VirtualProduction/Takes/Source/TakesCore/Private/TakeMetaData.cpp b/Engine/Plugins/VirtualProduction/Takes/Source/TakesCore/Private/TakeMetaData.cpp
--- a/Engine/Plugins/VirtualProduction/Takes/Source/TakesCore/Private/TakeMetaData.cpp
+++ b/Engine/Plugins/VirtualProduction/Takes/Source/TakesCore/Private/TakeMetaData.cpp
@@ -253,6 +253,7 @@ void UTakeMetaData::SetSlate(FString InSlate, bool bEmitChanged)
 		if (bEmitChanged)
 		{
 			UTakesCoreBlueprintLibrary::OnTakeRecorderSlateChanged(Slate);
+			OnTakeSlateChanged().Broadcast(Slate, this);
 		}
 	}
 }
@@ -266,6 +267,7 @@ void UTakeMetaData::SetTakeNumber(int32 InTakeNumber, bool bEmitChanged)
 		if (bEmitChanged)
 		{
 			UTakesCoreBlueprintLibrary::OnTakeRecorderTakeNumberChanged(TakeNumber);
+			OnTakeNumberChanged().Broadcast(TakeNumber, this);
 		}
 	}
 }
@@ -357,6 +359,17 @@ void UTakeMetaData::SetIsSubSequence(bool bNewValue)
 	}
 }
 
+UTakeMetaData::FOnTakeSlateChanged& UTakeMetaData::OnTakeSlateChanged()
+{
+	static FOnTakeSlateChanged OnSlateChanged;
+	return OnSlateChanged;
+}
+
+UTakeMetaData::FOnTakeNumberChanged& UTakeMetaData::OnTakeNumberChanged()
+{
+	static FOnTakeNumberChanged OnTakeNumberChanged;
+	return OnTakeNumberChanged;
+}
 
 void UTakeMetaData::ExtendAssetRegistryTags(FAssetRegistryTagsContext Context) const
 {
diff --git a/Engine/Plugins/VirtualProduction/Takes/Source/TakesCore/Public/TakeMetaData.h b/Engine/Plugins/VirtualProduction/Takes/Source/TakesCore/Public/TakeMetaData.h
--- a/Engine/Plugins/VirtualProduction/Takes/Source/TakesCore/Public/TakeMetaData.h
+++ b/Engine/Plugins/VirtualProduction/Takes/Source/TakesCore/Public/TakeMetaData.h
@@ -291,6 +291,14 @@ public:
 	UFUNCTION(BlueprintCallable, Category = "Take")
 	void SetIsSubSequence(bool bNewValue);
 
+	DECLARE_MULTICAST_DELEGATE_TwoParams(FOnTakeSlateChanged, const FString& SlateName, UTakeMetaData* TakeMetaData);
+	/** Multicast delegate emitted when the slate is changed. */
+	static FOnTakeSlateChanged& OnTakeSlateChanged();
+
+	DECLARE_MULTICAST_DELEGATE_TwoParams(FOnTakeNumberChanged, int32 TakeNumber, UTakeMetaData* TakeMetaData);
+	/** Multicast delegate emitted when take number is changed. */
+	static FOnTakeNumberChanged& OnTakeNumberChanged();
+	
 private:
 
 	/** Whether the take is locked */
