TakeRecorder: Fix crash when completing a recording using the Level Sequencer.

diff --git a/Engine/Plugins/VirtualProduction/Takes/Source/TakeRecorder/Private/Recorder/TakeRecorderSubsystem.cpp b/Engine/Plugins/VirtualProduction/Takes/Source/TakeRecorder/Private/Recorder/TakeRecorderSubsystem.cpp
--- a/Engine/Plugins/VirtualProduction/Takes/Source/TakeRecorder/Private/Recorder/TakeRecorderSubsystem.cpp
+++ b/Engine/Plugins/VirtualProduction/Takes/Source/TakeRecorder/Private/Recorder/TakeRecorderSubsystem.cpp
@@ -945,11 +945,14 @@ void UTakeRecorderSubsystem::IncrementTakeNumber()
 	}
 	
 	// Update the preset take number at the end of recording
-	const ULevelSequence* LevelSequence = TransientPreset->GetLevelSequence();
-	if (UTakeMetaData* MetaData = LevelSequence ? LevelSequence->FindMetaData<UTakeMetaData>() : nullptr)
+	
+	if (const ULevelSequence* LevelSequence = TransientPreset ? TransientPreset->GetLevelSequence() : nullptr)
 	{
-		const int32 NextTakeNumber = GetNextTakeNumber(MetaData->GetSlate());
-		MetaData->SetTakeNumber(NextTakeNumber);
+		if (UTakeMetaData* MetaData = LevelSequence->FindMetaData<UTakeMetaData>())
+		{
+			const int32 NextTakeNumber = GetNextTakeNumber(MetaData->GetSlate());
+			MetaData->SetTakeNumber(NextTakeNumber);
+		}
 	}
 
 	bAutoApplyTakeNumber = true;
