TakeRecorder: Fix SetTakeNumber in the subsystem not leading to OnTakeNumberChanged event being called.

diff --git a/Engine/Plugins/VirtualProduction/Takes/Source/TakeRecorder/Private/Recorder/TakeRecorderSubsystem.cpp b/Engine/Plugins/VirtualProduction/Takes/Source/TakeRecorder/Private/Recorder/TakeRecorderSubsystem.cpp
--- a/Engine/Plugins/VirtualProduction/Takes/Source/TakeRecorder/Private/Recorder/TakeRecorderSubsystem.cpp
+++ b/Engine/Plugins/VirtualProduction/Takes/Source/TakeRecorder/Private/Recorder/TakeRecorderSubsystem.cpp
@@ -538,13 +538,13 @@ ETakeRecorderState UTakeRecorderSubsystem::GetState() const
 	return ETakeRecorderState::PreInitialization;
 }
 
-void UTakeRecorderSubsystem::SetTakeNumber(int32 InNewTakeNumber)
+void UTakeRecorderSubsystem::SetTakeNumber(int32 InNewTakeNumber, bool bEmitChanged)
 {
 	if (TakeMetaData)
 	{
 		FScopedTransaction Transaction(LOCTEXT("SetTake_Transaction", "Set Take Number"));
 		TakeMetaData->Modify();
-		TakeMetaData->SetTakeNumber(InNewTakeNumber, false);
+		TakeMetaData->SetTakeNumber(InNewTakeNumber, bEmitChanged);
 		bAutoApplyTakeNumber = false;
 	}
 }
@@ -594,20 +594,20 @@ TArray<FAssetData> UTakeRecorderSubsystem::GetSlates(FName InPackagePath) const
 	return AllAssets;
 }
 
-void UTakeRecorderSubsystem::SetSlateName(const FString& InSlateName)
+void UTakeRecorderSubsystem::SetSlateName(const FString& InSlateName, bool bEmitChanged)
 {
 	if (TakeMetaData && TakeMetaData->GetSlate() != InSlateName)
 	{
 		FScopedTransaction Transaction(LOCTEXT("SetSlate_Transaction", "Set Take Slate"));
 		TakeMetaData->Modify();
 
-		TakeMetaData->SetSlate(InSlateName);
+		TakeMetaData->SetSlate(InSlateName, bEmitChanged);
 
 		// Compute the correct starting take number
 		const int32 NextTakeNumber = GetNextTakeNumber(TakeMetaData->GetSlate());
 		if (NextTakeNumber != TakeMetaData->GetTakeNumber())
 		{
-			TakeMetaData->SetTakeNumber(NextTakeNumber);
+			TakeMetaData->SetTakeNumber(NextTakeNumber, bEmitChanged);
 		}
 	}
 }
diff --git a/Engine/Plugins/VirtualProduction/Takes/Source/TakeRecorder/Private/Widgets/STakeRecorderCockpit.cpp b/Engine/Plugins/VirtualProduction/Takes/Source/TakeRecorder/Private/Widgets/STakeRecorderCockpit.cpp
--- a/Engine/Plugins/VirtualProduction/Takes/Source/TakeRecorder/Private/Widgets/STakeRecorderCockpit.cpp
+++ b/Engine/Plugins/VirtualProduction/Takes/Source/TakeRecorder/Private/Widgets/STakeRecorderCockpit.cpp
@@ -897,7 +897,9 @@ void STakeRecorderCockpit::SetTakeNumber(int32 InNewTakeNumber)
 
 	if (TransactionIndex != INDEX_NONE || bIsInPIEOrSimulate)
 	{
-		TakeRecorderSubsystem->SetTakeNumber(InNewTakeNumber);
+		// Don't emit here, will be done later.
+		constexpr bool bEmitChanged = false;
+		TakeRecorderSubsystem->SetTakeNumber(InNewTakeNumber, bEmitChanged);
 	}
 	
 	OnTokenValueUpdated();
@@ -918,7 +920,7 @@ void STakeRecorderCockpit::SetTakeNumber_FromCommit(int32 InNewTakeNumber, EText
 	}
 	else if (TakeMetaData->GetTakeNumber() != InNewTakeNumber)
 	{
-		TakeMetaData->SetTakeNumber(InNewTakeNumber);
+		TakeRecorderSubsystem->SetTakeNumber(InNewTakeNumber, /* bEmitChanged */ true);
 	}
 
 	OnTokenValueUpdated();
@@ -932,8 +934,7 @@ void STakeRecorderCockpit::OnEndSetTakeNumber(int32 InFinalValue)
 	{
 		check(TransactionIndex != INDEX_NONE);
 	}
-	UTakeMetaData* TakeMetaData = GetMetaDataChecked();
-	TakeMetaData->SetTakeNumber(InFinalValue);
+	TakeRecorderSubsystem->SetTakeNumber(InFinalValue, /* bEmitChanged */ true);
 
 	GEditor->EndTransaction();
 	TransactionIndex = INDEX_NONE;
diff --git a/Engine/Plugins/VirtualProduction/Takes/Source/TakeRecorder/Public/Recorder/TakeRecorderSubsystem.h b/Engine/Plugins/VirtualProduction/Takes/Source/TakeRecorder/Public/Recorder/TakeRecorderSubsystem.h
--- a/Engine/Plugins/VirtualProduction/Takes/Source/TakeRecorder/Public/Recorder/TakeRecorderSubsystem.h
+++ b/Engine/Plugins/VirtualProduction/Takes/Source/TakeRecorder/Public/Recorder/TakeRecorderSubsystem.h
@@ -178,9 +178,13 @@ public:
 	UFUNCTION(BlueprintCallable, Category = TakeRecorder)
 	ETakeRecorderState GetState() const;
 
-	/** Directly set the take number. */
+	/**
+	 * Directly set the take number.
+	 * @param InNewTakeNumber The new take number to set.
+	 * @param bEmitChanged Whether to broadcast events signaling the take number has changed.
+	 */
 	UFUNCTION(BlueprintCallable, Category = TakeRecorder)
-	void SetTakeNumber(int32 InNewTakeNumber);
+	void SetTakeNumber(int32 InNewTakeNumber, bool bEmitChanged = true);
 
 	/** Compute the next take number given a slate. */
 	UFUNCTION(BlueprintCallable, Category = TakeRecorder)
@@ -197,9 +201,13 @@ public:
 	UFUNCTION(BlueprintCallable, Category = TakeRecorder)
 	TArray<FAssetData> GetSlates(FName InPackagePath = NAME_None) const;
 	
-	/** Directly set the slate name. */
+	/**
+	 * Directly set the slate name.
+	 * @param InSlateName The new slate name to set.
+	 * @param bEmitChanged Whether to broadcast events signaling the slate name has changed.
+	 */
 	UFUNCTION(BlueprintCallable, Category = TakeRecorder)
-	void SetSlateName(const FString& InSlateName);
+	void SetSlateName(const FString& InSlateName, bool bEmitChanged = true);
 
 	/** Mark the current frame. */
 	UFUNCTION(BlueprintCallable, Category = TakeRecorder)
