TakeRecorder: Improve object reference management in subsystem and tokens. Clear cached TakeMetaData from the editor subsystem if user is force deleting the contained level sequence. This allows users to safely delete a level sequence which has just been recorded.

diff --git a/Engine/Plugins/VirtualProduction/Takes/Source/TakeRecorder/Private/Recorder/TakeRecorderSubsystem.cpp b/Engine/Plugins/VirtualProduction/Takes/Source/TakeRecorder/Private/Recorder/TakeRecorderSubsystem.cpp
--- a/Engine/Plugins/VirtualProduction/Takes/Source/TakeRecorder/Private/Recorder/TakeRecorderSubsystem.cpp
+++ b/Engine/Plugins/VirtualProduction/Takes/Source/TakeRecorder/Private/Recorder/TakeRecorderSubsystem.cpp
@@ -62,6 +62,11 @@ void UTakeRecorderSubsystem::Deinitialize()
 
 	UTakeMetaData::OnTakeSlateChanged().RemoveAll(this);
 	UTakeMetaData::OnTakeNumberChanged().RemoveAll(this);
+
+	UTakePresetSettings::Get()->OnSettingsChanged().RemoveAll(this);
+	
+	FEditorDelegates::OnPreForceDeleteObjects.Remove(OnPreForceDeleteObjectsHandle);
+	OnPreForceDeleteObjectsHandle.Reset();
 	
 	if (const FAssetRegistryModule* AssetRegistryModule = FModuleManager::GetModulePtr<FAssetRegistryModule>("AssetRegistry"))
 	{
@@ -248,9 +253,9 @@ bool UTakeRecorderSubsystem::CanReviewLastRecording() const
 
 bool UTakeRecorderSubsystem::ReviewLastRecording()
 {
-	if (LastRecordedLevelSequence)
+	if (LastRecordedLevelSequence.IsValid())
 	{
-		SuppliedLevelSequence = LastRecordedLevelSequence;
+		SuppliedLevelSequence = LastRecordedLevelSequence.Get();
 		LastRecordedLevelSequence = nullptr;
 		return true;
 	}
@@ -782,7 +787,7 @@ ULevelSequence* UTakeRecorderSubsystem::GetRecordIntoLevelSequence() const
 
 ULevelSequence* UTakeRecorderSubsystem::GetLastRecordedLevelSequence() const
 {
-	return LastRecordedLevelSequence;
+	return LastRecordedLevelSequence.Get();
 }
 
 UTakePreset* UTakeRecorderSubsystem::GetTransientPreset() const
@@ -903,6 +908,13 @@ void UTakeRecorderSubsystem::CacheMetaData()
 	if (NewMetaDataThisTick != TakeMetaData)
 	{
 		TakeMetaData = NewMetaDataThisTick;
+
+		if (!OnPreForceDeleteObjectsHandle.IsValid())
+		{
+			// Since this is a strong reference, we need to be able to clear it if the user is attempting to remove it, but let's not always
+			// hook this event since the subsystem won't be active for general UE operation.
+			OnPreForceDeleteObjectsHandle = FEditorDelegates::OnPreForceDeleteObjects.AddUObject(this, &UTakeRecorderSubsystem::OnPreForceDeleteObjects);
+		}
 	}
 
 	//Set MovieScene Display Rate to the Preset Frame Rate.
@@ -969,6 +981,12 @@ void UTakeRecorderSubsystem::OnAssetRegistryFilesLoaded()
 
 void UTakeRecorderSubsystem::OnRecordingInitialized(UTakeRecorder* Recorder)
 {
+	if (!bHasTargetSequenceBeenSet)
+	{
+		// Not initialized, take recorder triggered outside of the subsystem.
+		SetTargetSequence();
+	}
+	
 	// This needs to be stored with a strong ptr before our panels refresh, otherwise the weak sequencer ref will be lost during initialize.
 	RecordingLevelSequence = Recorder->GetSequence();
 	// Recache the meta-data here since we know that the sequence has probably changed as a result of the recording being started
@@ -1016,7 +1034,7 @@ void UTakeRecorderSubsystem::OnRecordingFinished(UTakeRecorder* Recorder)
 	Recorder->OnRecordingFinished().Remove(OnRecordingFinishedHandle);
 
 	GetOnRecordingFinishedEvent().Broadcast(Recorder);
-	TakeRecorderFinished.Broadcast(LastRecordedLevelSequence);
+	TakeRecorderFinished.Broadcast(LastRecordedLevelSequence.Get());
 }
 
 void UTakeRecorderSubsystem::OnRecordingCancelled(UTakeRecorder* Recorder)
@@ -1112,4 +1130,28 @@ void UTakeRecorderSubsystem::OnTakePresetSettingsChanged()
 	}
 }
 
+void UTakeRecorderSubsystem::OnPreForceDeleteObjects(const TArray<UObject*>& InObjects)
+{
+	// Take meta data is cached as a strong reference, but can prevent us from deleting owning level sequences.
+	// We need to keep a strong reference to it for Take Recorder functionality, especially in the panel, and it should persist
+	// even with the panel closed. Listening for a force delete allows us to safely clear the reference.
+	const ULevelSequence* LevelSequence = GetLevelSequence();
+	if (!LevelSequence && LastRecordedLevelSequence.IsValid())
+	{
+		LevelSequence = LastRecordedLevelSequence.Get();
+	}
+	if (LevelSequence && TakeMetaData && InObjects.Contains(LevelSequence)
+		&& LevelSequence->FindMetaData<UTakeMetaData>() == TakeMetaData)
+	{
+		if (IsRecording())
+		{
+			// Need to cancel or we will be in a bad state. TakeRecorder will tick but be without a level sequence.
+			CancelRecording(false);
+		}
+		TakeMetaData = nullptr;
+		FEditorDelegates::OnPreForceDeleteObjects.Remove(OnPreForceDeleteObjectsHandle);
+		OnPreForceDeleteObjectsHandle.Reset();
+	}
+}
+
 #undef LOCTEXT_NAMESPACE
diff --git a/Engine/Plugins/VirtualProduction/Takes/Source/TakeRecorder/Public/Recorder/TakeRecorderSubsystem.h b/Engine/Plugins/VirtualProduction/Takes/Source/TakeRecorder/Public/Recorder/TakeRecorderSubsystem.h
--- a/Engine/Plugins/VirtualProduction/Takes/Source/TakeRecorder/Public/Recorder/TakeRecorderSubsystem.h
+++ b/Engine/Plugins/VirtualProduction/Takes/Source/TakeRecorder/Public/Recorder/TakeRecorderSubsystem.h
@@ -380,6 +380,9 @@ private:
 	
 	/** Handles the UTakePresetSettings::RecordTargetClass changing: Recreates the transient level sequence if recording a transaction. */
 	void OnTakePresetSettingsChanged();
+
+	/** Callback when the engine is attempting to force delete objects. */
+	void OnPreForceDeleteObjects(const TArray<UObject*>& InObjects);
 	
 private:
 	/** Last data used for initialization. */
@@ -403,7 +406,7 @@ private:
 
 	/** The last recorded level sequence. */
 	UPROPERTY()
-	TObjectPtr<ULevelSequence> LastRecordedLevelSequence;
+	TWeakObjectPtr<ULevelSequence> LastRecordedLevelSequence;
 
 	/** Take meta-data cached from the level sequence if it exists. */
 	UPROPERTY()
@@ -442,6 +445,7 @@ private:
 	
 	FDelegateHandle OnAssetRegistryFilesLoadedHandle;
 	FDelegateHandle OnRecordingInitializedHandle, OnRecordingStoppedHandle, OnRecordingFinishedHandle, OnRecordingCancelledHandle;
+	FDelegateHandle OnPreForceDeleteObjectsHandle;
 	
 	/** Whether we should auto apply the next available take number when asset discovery has finished or not */
 	bool bAutoApplyTakeNumber;
diff --git a/Engine/Plugins/VirtualProduction/Takes/Source/TakeRecorderNamingTokens/Private/TakeRecorderNamingTokens.cpp b/Engine/Plugins/VirtualProduction/Takes/Source/TakeRecorderNamingTokens/Private/TakeRecorderNamingTokens.cpp
--- a/Engine/Plugins/VirtualProduction/Takes/Source/TakeRecorderNamingTokens/Private/TakeRecorderNamingTokens.cpp
+++ b/Engine/Plugins/VirtualProduction/Takes/Source/TakeRecorderNamingTokens/Private/TakeRecorderNamingTokens.cpp
@@ -113,7 +113,7 @@ void UTakeRecorderNamingTokens::OnCreateDefaultTokens(TArray<FNamingTokenData>&
 					FString AssetName = Sequence->GetName();
 					
 					const UTakePreset* TransientPreset = Subsystem->GetTransientPreset();
-					if (TransientPreset && Sequence == TransientPreset->GetLevelSequence() && TakeMetaData)
+					if (TransientPreset && Sequence == TransientPreset->GetLevelSequence() && TakeMetaData.IsValid())
 					{
 						// Working with the transient preset... try to generate the most correct takeName.
 						// Depending on when this is read, it may not match the actual takeName (Sequencer asset's name)
@@ -154,7 +154,7 @@ void UTakeRecorderNamingTokens::OnCreateDefaultTokens(TArray<FNamingTokenData>&
 		TEXT("actor"),
 		LOCTEXT("TokenActor", "Actor"),
 		FNamingTokenData::FTokenProcessorDelegateNative::CreateLambda([this]{
-			if (Context && Context->Actor)
+			if (Context && Context->Actor.IsValid())
 			{
 				return FText::FromString(Context->Actor->GetActorLabel());
 			}
@@ -191,7 +191,7 @@ void UTakeRecorderNamingTokens::OnPreEvaluate_Implementation(const FNamingTokens
 
 	// Use either provided take meta-data (manually evaluated from take recorder), or locate the most recent (global operation)
 	TakeMetaData = (Context && Context->TakeMetaData != nullptr) ?
-		Context->TakeMetaData : TObjectPtr<UTakeMetaData>(UTakeMetaData::GetMostRecentMetaData());
+		Context->TakeMetaData : TWeakObjectPtr<UTakeMetaData>(UTakeMetaData::GetMostRecentMetaData());
 	if (TakeMetaData == nullptr)
 	{
 		// Create metadata which computes its data from available information.
@@ -230,7 +230,7 @@ void UTakeRecorderNamingTokens::OnPostEvaluate_Implementation()
 
 FDateTime UTakeRecorderNamingTokens::GetCurrentDateTime_Implementation() const
 {
-	const FDateTime TakeTimestamp = TakeMetaData ? TakeMetaData->GetTimestamp() : FDateTime(0);
+	const FDateTime TakeTimestamp = TakeMetaData.IsValid() ? TakeMetaData->GetTimestamp() : FDateTime(0);
 	return TakeTimestamp == FDateTime(0) ? Super::GetCurrentDateTime_Implementation() : TakeTimestamp;
 }
 
diff --git a/Engine/Plugins/VirtualProduction/Takes/Source/TakeRecorderNamingTokens/Private/TakeRecorderNamingTokens.h b/Engine/Plugins/VirtualProduction/Takes/Source/TakeRecorderNamingTokens/Private/TakeRecorderNamingTokens.h
--- a/Engine/Plugins/VirtualProduction/Takes/Source/TakeRecorderNamingTokens/Private/TakeRecorderNamingTokens.h
+++ b/Engine/Plugins/VirtualProduction/Takes/Source/TakeRecorderNamingTokens/Private/TakeRecorderNamingTokens.h
@@ -32,7 +32,7 @@ protected:
 private:
 	/** Cached metadata for this run. */
 	UPROPERTY(Transient)
-	TObjectPtr<const UTakeMetaData> TakeMetaData;
+	TWeakObjectPtr<const UTakeMetaData> TakeMetaData;
 
 	/** Cached context for this run. This isn't available globally and requires a context passed to it. */
 	UPROPERTY(Transient)
diff --git a/Engine/Plugins/VirtualProduction/Takes/Source/TakesCore/Public/NamingTokens/TakeRecorderNamingTokensContext.h b/Engine/Plugins/VirtualProduction/Takes/Source/TakesCore/Public/NamingTokens/TakeRecorderNamingTokensContext.h
--- a/Engine/Plugins/VirtualProduction/Takes/Source/TakesCore/Public/NamingTokens/TakeRecorderNamingTokensContext.h
+++ b/Engine/Plugins/VirtualProduction/Takes/Source/TakesCore/Public/NamingTokens/TakeRecorderNamingTokensContext.h
@@ -23,11 +23,11 @@ class UTakeRecorderNamingTokensContext : public UObject
 public:
 	/** MetaData specifically for this context. Setting this prevents having to perform a global lookup. */
 	UPROPERTY(Transient)
-	TObjectPtr<const UTakeMetaData> TakeMetaData;
+	TWeakObjectPtr<const UTakeMetaData> TakeMetaData;
 	
 	/** The specific actor for this context. @todo NamingTokens - Determine if we should be retrieving the {actor} a different way. */
 	UPROPERTY(Transient)
-	TObjectPtr<const AActor> Actor;
+	TWeakObjectPtr<const AActor> Actor;
 	
 	/** The audio device channel for this context. @todo NamingTokens - Determine if we should be retrieving the {channel} a different way. */
 	UPROPERTY(Transient)
