TakeRecorder: Fix subsequence asset name and track not being utilized in standard actor sources. Add Subsequence directory field, recognized by [REDACTED_PERSON], to be used when sub sequences are recorded.

diff --git a/Engine/Plugins/VirtualProduction/Takes/Source/TakeRecorder/Private/Recorder/TakeRecorder.cpp b/Engine/Plugins/VirtualProduction/Takes/Source/TakeRecorder/Private/Recorder/TakeRecorder.cpp
--- a/Engine/Plugins/VirtualProduction/Takes/Source/TakeRecorder/Private/Recorder/TakeRecorder.cpp
+++ b/Engine/Plugins/VirtualProduction/Takes/Source/TakeRecorder/Private/Recorder/TakeRecorder.cpp
@@ -878,6 +878,7 @@ void UTakeRecorder::PreRecord()
 	TakeRecorderSourcesSettings.bSaveRecordedAssets = (!Parameters.bDisableRecordingAndSave && Parameters.User.bSaveRecordedAssets) || GEditor == nullptr;
 	TakeRecorderSourcesSettings.bRemoveRedundantTracks = Parameters.User.bRemoveRedundantTracks;
 	TakeRecorderSourcesSettings.bAutoLock = Parameters.User.bAutoLock;
+	TakeRecorderSourcesSettings.SubsequenceDirectory = Parameters.Project.SubsequenceDirectory;
 
 	TSharedPtr<ISequencer> Sequencer = WeakSequencer.Pin();
 	if (Sequencer)
diff --git a/Engine/Plugins/VirtualProduction/Takes/Source/TakeRecorder/Private/Recorder/TakeRecorderParameters.cpp b/Engine/Plugins/VirtualProduction/Takes/Source/TakeRecorder/Private/Recorder/TakeRecorderParameters.cpp
--- a/Engine/Plugins/VirtualProduction/Takes/Source/TakeRecorder/Private/Recorder/TakeRecorderParameters.cpp
+++ b/Engine/Plugins/VirtualProduction/Takes/Source/TakeRecorder/Private/Recorder/TakeRecorderParameters.cpp
@@ -24,7 +24,8 @@ FTakeRecorderUserParameters::FTakeRecorderUserParameters()
 }
 
 FTakeRecorderProjectParameters::FTakeRecorderProjectParameters()
-	: RecordingClockSource(EUpdateClockSource::RelativeTimecode)
+	: SubsequenceDirectory(TEXT("{takeName}_Subscenes"))
+	, RecordingClockSource(EUpdateClockSource::RelativeTimecode)
 	, bStartAtCurrentTimecode(true)
 	, bRecordTimecode(false)
 	, bRecordSourcesIntoSubSequences(false)
diff --git a/Engine/Plugins/VirtualProduction/Takes/Source/TakeRecorder/Public/Recorder/TakeRecorderParameters.h b/Engine/Plugins/VirtualProduction/Takes/Source/TakeRecorder/Public/Recorder/TakeRecorderParameters.h
--- a/Engine/Plugins/VirtualProduction/Takes/Source/TakeRecorder/Public/Recorder/TakeRecorderParameters.h
+++ b/Engine/Plugins/VirtualProduction/Takes/Source/TakeRecorder/Public/Recorder/TakeRecorderParameters.h
@@ -86,6 +86,12 @@ struct FTakeRecorderProjectParameters
 	UPROPERTY(config, EditAnywhere, BlueprintReadWrite, Category = "Take Recorder", meta = (NamingTokens))
 	FString TakeSaveDir;
 
+	/**
+	 * The directory to record subscenes into.
+	 */
+	UPROPERTY(config, EditAnywhere, BlueprintReadWrite, Category = "Take Recorder", meta = (NamingTokens))
+	FString SubsequenceDirectory;
+	
 	/**
 	 * The default name to use for the Slate information
 	 */
diff --git a/Engine/Plugins/VirtualProduction/Takes/Source/TakeRecorderSources/Private/TakeRecorderActorSource.cpp b/Engine/Plugins/VirtualProduction/Takes/Source/TakeRecorderSources/Private/TakeRecorderActorSource.cpp
--- a/Engine/Plugins/VirtualProduction/Takes/Source/TakeRecorderSources/Private/TakeRecorderActorSource.cpp
+++ b/Engine/Plugins/VirtualProduction/Takes/Source/TakeRecorderSources/Private/TakeRecorderActorSource.cpp
@@ -36,6 +36,7 @@
 #include "TrackRecorders/MovieSceneTrackRecorder.h"
 #include "TrackRecorders/MovieScenePropertyTrackRecorder.h"
 #include "TrackRecorders/MovieSceneAnimationTrackRecorder.h"
+#include "TrackRecorders/MovieSceneAnimationTrackRecorderSettings.h"
 #include "TrackRecorders/MovieScene3DTransformTrackRecorder.h"
 #include "TrackRecorders/MovieScenePropertyTrackRecorder.h"
 #include "TrackRecorders/MovieSceneTrackRecorderSettings.h"
@@ -47,6 +48,7 @@
 #if WITH_EDITOR
 #include "Tracks/MovieSceneSpawnTrack.h"
 #endif
+#include "NamingTokens/TakeRecorderNamingTokensContext.h"
 #include "SequencerUtilities.h"
 #include "TakeRecorderSourcesModule.h"
 
@@ -728,6 +730,19 @@ void UTakeRecorderActorSource::ProcessRecordedTimes(ULevelSequence* InSequence)
 	TakeRecorderSourceHelpers::ProcessRecordedTimes(InSequence, TakeTrack, FrameRange, UTakeRecorderSources::RecordedTimes);
 }
 
+UMovieSceneAnimationTrackRecorderEditorSettings* UTakeRecorderActorSource::GetBestTrackRecorderEditorSettings() const
+{
+	// This is called when the instance settings most likely don't exist, but we should still check them.
+	// The animation track recorder editor settings are shared with normal actor (non-animation) usage for asset and track name.
+	if (UMovieSceneAnimationTrackRecorderEditorSettings* InstanceSettings
+		= Cast<UMovieSceneAnimationTrackRecorderSettings>(GetSettingsObjectForFactory(UMovieSceneAnimationTrackRecorderSettings::StaticClass())))
+	{
+		return InstanceSettings;
+	}
+	// Fallback to global.
+	return GetMutableDefault<UMovieSceneAnimationTrackRecorderEditorSettings>();
+}
+
 TArray<UTakeRecorderSource*> UTakeRecorderActorSource::PostRecording(ULevelSequence* InSequence, class ULevelSequence* InRootSequence, const bool bCancelled)
 {
 	FTakeRecorderParameters Parameters;
@@ -1787,6 +1802,16 @@ UMovieSceneTrackRecorderSettings* UTakeRecorderActorSource::GetSettingsObjectFor
 
 FString UTakeRecorderActorSource::GetSubsceneTrackName(ULevelSequence* InSequence) const
 {
+	if (const UTakeMetaData* TakeMetaData = InSequence->FindMetaData<UTakeMetaData>())
+	{
+		if (const UMovieSceneAnimationTrackRecorderEditorSettings* Settings = GetBestTrackRecorderEditorSettings())
+		{
+			UTakeRecorderNamingTokensContext* Context = NewObject<UTakeRecorderNamingTokensContext>();
+			Context->Actor = Target.LoadSynchronous();
+			return TakeMetaData->GenerateAssetPath(Settings->AnimationTrackName.ToString(), Context);
+		}
+	}
+	
 	if (Target.IsValid())
 	{
 		return Target->GetActorLabel();
@@ -1797,12 +1822,18 @@ FString UTakeRecorderActorSource::GetSubsceneTrackName(ULevelSequence* InSequenc
 
 FString UTakeRecorderActorSource::GetSubsceneAssetName(ULevelSequence* InSequence) const
 {
-	if (Target.IsValid())
+	if (const UTakeMetaData* TakeMetaData = InSequence->FindMetaData<UTakeMetaData>())
 	{
-		if (UTakeMetaData* TakeMetaData = InSequence->FindMetaData<UTakeMetaData>())
+		if (const UMovieSceneAnimationTrackRecorderEditorSettings* Settings = GetBestTrackRecorderEditorSettings())
 		{
-			return FString::Printf(TEXT("%s_%s"), *Target->GetActorLabel(), *TakeMetaData->GenerateAssetPath("{slate}_{take}"));
+			UTakeRecorderNamingTokensContext* Context = NewObject<UTakeRecorderNamingTokensContext>();
+			Context->Actor = Target.LoadSynchronous();
+			return TakeMetaData->GenerateAssetPath(Settings->AnimationAssetName, Context);
 		}
+	}
+	
+	if (Target.IsValid())
+	{
 		return Target->GetActorLabel();
 	}
 
diff --git a/Engine/Plugins/VirtualProduction/Takes/Source/TakeRecorderSources/Private/TakeRecorderActorSource.h b/Engine/Plugins/VirtualProduction/Takes/Source/TakeRecorderSources/Private/TakeRecorderActorSource.h
--- a/Engine/Plugins/VirtualProduction/Takes/Source/TakeRecorderSources/Private/TakeRecorderActorSource.h
+++ b/Engine/Plugins/VirtualProduction/Takes/Source/TakeRecorderSources/Private/TakeRecorderActorSource.h
@@ -15,6 +15,7 @@ class UActorComponent;
 class UMovieSceneFolder;
 class UMovieScene;
 class UMovieSceneTrackRecorderSettings;
+class UMovieSceneAnimationTrackRecorderEditorSettings;
 
 DECLARE_LOG_CATEGORY_EXTERN(ActorSerialization, Verbose, All);
 
@@ -300,6 +301,9 @@ private:
 	
 	void ProcessRecordedTimes(ULevelSequence* InSequence);
 
+	/** Return either instance or global animation track recorder editor settings. */
+	UMovieSceneAnimationTrackRecorderEditorSettings* GetBestTrackRecorderEditorSettings() const;
+
 private:
 	/** Object Binding guid that is created in the Level Sequence when recording starts.*/
 	FGuid CachedObjectBindingGuid;
diff --git a/Engine/Plugins/VirtualProduction/Takes/Source/TakesCore/Private/TakeRecorderSources.cpp b/Engine/Plugins/VirtualProduction/Takes/Source/TakesCore/Private/TakeRecorderSources.cpp
--- a/Engine/Plugins/VirtualProduction/Takes/Source/TakesCore/Private/TakeRecorderSources.cpp
+++ b/Engine/Plugins/VirtualProduction/Takes/Source/TakesCore/Private/TakeRecorderSources.cpp
@@ -20,6 +20,7 @@
 #include "Engine/TimecodeProvider.h"
 #include "Misc/App.h"
 #include "Editor.h"
+#include "NamingTokens/TakeRecorderNamingTokensContext.h"
 #include "UObject/Package.h"
 
 #include UE_INLINE_GENERATED_CPP_BY_NAME(TakeRecorderSources)
@@ -129,7 +130,7 @@ void UTakeRecorderSources::PreRecordingRecursive(TArray<UTakeRecorderSource*> In
 				const FString& SubSequenceTrackName = ObjectTools::SanitizeObjectName(Source->GetSubsceneTrackName(InRootSequence));
 				const FString& SubSequenceAssetName = ObjectTools::SanitizeObjectName(Source->GetSubsceneAssetName(InRootSequence));
 
-				TargetSequence = CreateSubSequenceForSource(InRootSequence, SubSequenceTrackName, SubSequenceAssetName);
+				TargetSequence = CreateSubSequenceForSource(InRootSequence, SubSequenceTrackName, SubSequenceAssetName, Settings);
 
 				// If there's already a Subscene Track for our sub-sequence we need to remove that track before create a new one. No data is lost in this process as the
 				// sequence that the subscene points to has been copied by CreateSubSequenceForSource so a new track pointed to the new subsequence includes all the old data.
@@ -546,7 +547,8 @@ void UTakeRecorderSources::StopRecording(class ULevelSequence* InSequence, const
 	}
 }
 
-ULevelSequence* UTakeRecorderSources::CreateSubSequenceForSource(ULevelSequence* InRootSequence, const FString& SubSequenceTrackName, const FString& SubSequenceAssetName)
+ULevelSequence* UTakeRecorderSources::CreateSubSequenceForSource(ULevelSequence* InRootSequence, const FString& SubSequenceTrackName,
+	const FString& SubSequenceAssetName, const FTakeRecorderSourcesSettings& InSettings)
 {
 	IAssetRegistry& AssetRegistry = FModuleManager::LoadModuleChecked<FAssetRegistryModule>("AssetRegistry").Get();
 
@@ -574,8 +576,24 @@ ULevelSequence* UTakeRecorderSources::CreateSubSequenceForSource(ULevelSequence*
 			}
 		}
 	}
-
-	FString NewPath = FString::Printf(TEXT("%s/%s_Subscenes/%s"), *SequenceDirectory, *SequenceName, *SubSequenceAssetName);
+	
+	FString SubsequenceDirectory;
+	if (const UTakeMetaData* AssetMetaData = InRootSequence->FindMetaData<UTakeMetaData>())
+	{
+		// User specified subsequence directory.
+		UTakeRecorderNamingTokensContext* Context = NewObject<UTakeRecorderNamingTokensContext>();
+		Context->TakeMetaData = AssetMetaData;
+		SubsequenceDirectory = InSettings.SubsequenceDirectory;
+		SubsequenceDirectory = AssetMetaData->GenerateAssetPath(SubsequenceDirectory, Context);
+	}
+	
+	if (SubsequenceDirectory.IsEmpty())
+	{
+		// Fallback to original behavior.
+		SubsequenceDirectory = FString::Printf(TEXT("%s_Subscenes"), *SequenceName);
+	}
+	
+	FString NewPath = FString::Printf(TEXT("%s/%s/%s"), *SequenceDirectory, *SubsequenceDirectory, *SubSequenceAssetName);
 
 	ULevelSequence* OutAsset = nullptr;
 	TakesUtils::CreateNewAssetPackage<ULevelSequence>(NewPath, OutAsset, nullptr, ExistingSubSequence, InRootSequence->GetClass());
diff --git a/Engine/Plugins/VirtualProduction/Takes/Source/TakesCore/Public/TakeRecorderSources.h b/Engine/Plugins/VirtualProduction/Takes/Source/TakesCore/Public/TakeRecorderSources.h
--- a/Engine/Plugins/VirtualProduction/Takes/Source/TakesCore/Public/TakeRecorderSources.h
+++ b/Engine/Plugins/VirtualProduction/Takes/Source/TakesCore/Public/TakeRecorderSources.h
@@ -25,6 +25,7 @@ struct FTakeRecorderSourcesSettings
 	bool bAutoLock;
 	bool bRemoveRedundantTracks;
 	FTransform TransformOrigin;
+	FString SubsequenceDirectory;
 };
 
 /**
@@ -173,7 +174,8 @@ public:
 	*/
 
 	/** Creates a sub-sequence asset for the specified sub sequence name based on the given root sequence. */
-	static ULevelSequence* CreateSubSequenceForSource(ULevelSequence* InRootSequence, const FString& SubSequenceTrackName, const FString& SubSequenceAssetName);
+	static ULevelSequence* CreateSubSequenceForSource(ULevelSequence* InRootSequence, const FString& SubSequenceTrackName,
+		const FString& SubSequenceAssetName, const FTakeRecorderSourcesSettings& InSettings = FTakeRecorderSourcesSettings());
 
 	/**
 	 * Array of pairs - key time and the corresponding timecode
