TakeRecorder: Fix recording into a Cine Assembly not incrementing the take number.

diff --git a/Engine/Plugins/VirtualProduction/Takes/Source/TakesCore/Private/TakesCoreBlueprintLibrary.cpp b/Engine/Plugins/VirtualProduction/Takes/Source/TakesCore/Private/TakesCoreBlueprintLibrary.cpp
--- a/Engine/Plugins/VirtualProduction/Takes/Source/TakesCore/Private/TakesCoreBlueprintLibrary.cpp
+++ b/Engine/Plugins/VirtualProduction/Takes/Source/TakesCore/Private/TakesCoreBlueprintLibrary.cpp
@@ -38,6 +38,7 @@ TArray<FAssetData> UTakesCoreBlueprintLibrary::FindTakes(const FString& Slate, i
 	FARFilter Filter;
 	Filter.ClassPaths.Add(ULevelSequence::StaticClass()->GetClassPathName());
 	Filter.TagsAndValues.Add(UTakeMetaData::AssetRegistryTag_Slate, Slate);
+	Filter.bRecursiveClasses = true; // Needed for other types of LevelSequences, like CineAssembly
 
 	TArray<FAssetData> AllAssets;
 	AssetRegistry.GetAssets(Filter, AllAssets);
