TakeRecorder: Refactor subsystem so it's a private implementation, allowing us more flexibility with hotfixes.

diff --git a/Engine/Plugins/VirtualProduction/Takes/Source/TakeRecorder/Private/Recorder/TakeRecorderSubsystem.cpp b/Engine/Plugins/VirtualProduction/Takes/Source/TakeRecorder/Private/Recorder/TakeRecorderSubsystem.cpp
--- a/Engine/Plugins/VirtualProduction/Takes/Source/TakeRecorder/Private/Recorder/TakeRecorderSubsystem.cpp
+++ b/Engine/Plugins/VirtualProduction/Takes/Source/TakeRecorder/Private/Recorder/TakeRecorderSubsystem.cpp
@@ -2,30 +2,11 @@
 
 #include "Recorder/TakeRecorderSubsystem.h"
 
-#include "AssetRegistry/AssetRegistryModule.h"
-#include "Dialog/SMessageDialog.h"
-#include "Editor/TransBuffer.h"
-#include "Framework/Notifications/NotificationManager.h"
-#include "ILevelSequenceEditorToolkit.h"
-#include "ITakeRecorderModule.h"
-#include "ITakeRecorderNamingTokensModule.h"
-#include "LevelSequence.h"
-#include "MovieScene.h"
-#include "NamingTokens.h"
-#include "NamingTokensEngineSubsystem.h"
 #include "Recorder/TakeRecorder.h"
-#include "Recorder/TakeRecorderBlueprintLibrary.h"
 #include "Recorder/TakeRecorderParameters.h"
-#include "ScopedTransaction.h"
-#include "TakeMetaData.h"
-#include "TakePresetSettings.h"
-#include "TakeRecorderModule.h"
-#include "TakeRecorderSettings.h"
 #include "TakeRecorderSource.h"
 #include "TakeRecorderSources.h"
-#include "TakesCoreBlueprintLibrary.h"
-#include "Widgets/Notifications/SNotificationList.h"
-#include "Widgets/TakeRecorderNamingTokensData.h"
+#include "TakeRecorderSubsystemImplementation.h"
 
 #define LOCTEXT_NAMESPACE "TakeRecorderSubsystem"
 
@@ -33,793 +14,263 @@ void UTakeRecorderSubsystem::Initialize(FSubsystemCollectionBase& Collection)
 {
 	Super::Initialize(Collection);
 
-	const FAssetRegistryModule& AssetRegistryModule = FModuleManager::LoadModuleChecked<FAssetRegistryModule>(TEXT("AssetRegistry"));
-	OnAssetRegistryFilesLoadedHandle = AssetRegistryModule.Get().OnFilesLoaded().AddUObject(this, &UTakeRecorderSubsystem::OnAssetRegistryFilesLoaded);
-	OnRecordingInitializedHandle = UTakeRecorder::OnRecordingInitialized().AddUObject(this, &UTakeRecorderSubsystem::OnRecordingInitialized);
-
-	UTakeRecorderSources::OnSourceAdded().AddUObject(this, &UTakeRecorderSubsystem::OnSourceAdded);
-	UTakeRecorderSources::OnSourceRemoved().AddUObject(this, &UTakeRecorderSubsystem::OnSourceRemoved);
-
-	UTakeMetaData::OnTakeSlateChanged().AddUObject(this, &UTakeRecorderSubsystem::OnTakeSlateChanged);
-	UTakeMetaData::OnTakeNumberChanged().AddUObject(this, &UTakeRecorderSubsystem::OnTakeNumberChanged);
-
-	NamingTokensData = NewObject<UTakeRecorderNamingTokensData>();
-	NamingTokensData->SetFlags(RF_Transactional);
-
-	BindToNamingTokenEvents();
-	
-	// If the UTakePresetSettings::TargetRecordClass changes, the level sequence in the TransientPreset must be regenerated.
-	UTakePresetSettings::Get()->OnSettingsChanged().AddUObject(this, &UTakeRecorderSubsystem::OnTakePresetSettingsChanged);
+	Implementation = NewObject<UTakeRecorderSubsystemImplementation>();
+	CastChecked<UTakeRecorderSubsystemImplementation>(Implementation.GetObject())->InitializeImplementation(this);
 }
 
 void UTakeRecorderSubsystem::Deinitialize()
 {
 	Super::Deinitialize();
-	UTakeRecorder::OnRecordingInitialized().Remove(OnRecordingInitializedHandle);
-
-	UTakeRecorderSources::OnSourceAdded().RemoveAll(this);
-	UTakeRecorderSources::OnSourceRemoved().RemoveAll(this);
-
-	UTakeMetaData::OnTakeSlateChanged().RemoveAll(this);
-	UTakeMetaData::OnTakeNumberChanged().RemoveAll(this);
-
-	UTakePresetSettings::Get()->OnSettingsChanged().RemoveAll(this);
-	
-	FEditorDelegates::OnPreForceDeleteObjects.Remove(OnPreForceDeleteObjectsHandle);
-	OnPreForceDeleteObjectsHandle.Reset();
-	
-	if (const FAssetRegistryModule* AssetRegistryModule = FModuleManager::GetModulePtr<FAssetRegistryModule>("AssetRegistry"))
-	{
-		if (IAssetRegistry* AssetRegistry = AssetRegistryModule->TryGet())
-		{
-			AssetRegistry->OnFilesLoaded().Remove(OnAssetRegistryFilesLoadedHandle);
-		}
-	}
 
-	if (ITakeRecorderModule* TakeRecorderModule = FModuleManager::Get().GetModulePtr<ITakeRecorderModule>("TakeRecorder"))
+	if (Implementation && Implementation.GetObject())
 	{
-		TakeRecorderModule->GetLastLevelSequenceProvider().Unbind();
+		CastChecked<UTakeRecorderSubsystemImplementation>(Implementation.GetObject())->DeinitializeImplementation();
 	}
-
-	UnbindNamingTokensEvents();
-
-	bHasTargetSequenceBeenSet = false;
-	SetTickableTickType(ETickableTickType::Never);
-}
-
-ETickableTickType UTakeRecorderSubsystem::GetTickableTickType() const
-{
-	// This is to prevent registration until we call SetTargetSequence.
-	return ETickableTickType::Never;
-}
-
-UWorld* UTakeRecorderSubsystem::GetTickableGameObjectWorld() const
-{
-	return FTickableGameObject::GetTickableGameObjectWorld();
-}
-
-bool UTakeRecorderSubsystem::IsTickable() const
-{
-	return bHasTargetSequenceBeenSet;
-}
-
-bool UTakeRecorderSubsystem::IsTickableWhenPaused() const
-{
-	return FTickableGameObject::IsTickableWhenPaused();
-}
-
-bool UTakeRecorderSubsystem::IsTickableInEditor() const
-{
-	return true;
-}
-
-void UTakeRecorderSubsystem::Tick(float DeltaTime)
-{
-	CacheMetaData();
-}
-
-TStatId UTakeRecorderSubsystem::GetStatId() const
-{
-	RETURN_QUICK_DECLARE_CYCLE_STAT(UTakeRecorderSubsystem, STATGROUP_Tickables);
 }
 
 void UTakeRecorderSubsystem::SetTargetSequence(const FTakeRecorderSequenceParameters& InData)
 {
-	// Don't initialize again if we already have the same data present.
-	if (TargetSequenceData.IsValid() && TargetSequenceData == InData)
-	{
-		return;
-	}
-	TargetSequenceData = InData;
-	
-	ITakeRecorderModule& TakeRecorderModule = FModuleManager::Get().LoadModuleChecked<ITakeRecorderModule>("TakeRecorder");
-	TakeRecorderModule.GetLastLevelSequenceProvider().BindUObject(this, &UTakeRecorderSubsystem::SetLastLevelSequence);
-
-	bAutoApplyTakeNumber = true;
-	
-	CacheMetaData();
-	
-	if (TakeMetaData && !TakeMetaData->IsLocked())
-	{
-		const int32 NextTakeNumber = GetNextTakeNumber(TakeMetaData->GetSlate());
-		if (NextTakeNumber != TakeMetaData->GetTakeNumber())
-		{
-			TakeMetaData->SetTakeNumber(NextTakeNumber);
-		}
-	}
-	
-	// If a recording is currently underway, initialize to that now
-	if (UTakeRecorder* ActiveRecorder = UTakeRecorder::GetActiveRecorder())
-	{
-		RecordingLevelSequence = ActiveRecorder->GetSequence();
-
-		ActiveRecorder->OnRecordingStopped().Remove(OnRecordingStoppedHandle);
-		ActiveRecorder->OnRecordingFinished().Remove(OnRecordingFinishedHandle);
-		ActiveRecorder->OnRecordingCancelled().Remove(OnRecordingCancelledHandle);
-		OnRecordingStoppedHandle = ActiveRecorder->OnRecordingStopped().AddUObject(this, &UTakeRecorderSubsystem::OnRecordingStopped);
-		OnRecordingFinishedHandle = ActiveRecorder->OnRecordingFinished().AddUObject(this, &UTakeRecorderSubsystem::OnRecordingFinished);
-		OnRecordingCancelledHandle = ActiveRecorder->OnRecordingCancelled().AddUObject(this, &UTakeRecorderSubsystem::OnRecordingCancelled);
-	}
-	else
-	{
-		RecordingLevelSequence = nullptr;
-	}
-	
-	TransientPreset = AllocateTransientPreset();
-	
-	// Copy the base preset into the transient preset if it was provided.
-	// We do this first so that anything that asks for its Level Sequence
-	// on construction gets the right one
-	if (InData.BasePreset)
-	{
-		TransientPreset->CopyFrom(InData.BasePreset);
-	}
-	else if (InData.BaseSequence)
-	{
-		TransientPreset->CopyFrom(InData.BaseSequence);
-
-		ULevelSequence* LevelSequence = TransientPreset->GetLevelSequence();
-
-		UTakeRecorderSources* BaseSources = InData.BaseSequence->FindMetaData<UTakeRecorderSources>();
-		if (BaseSources && LevelSequence)
-		{
-			LevelSequence->CopyMetaData(BaseSources);
-		}
-
-		if (LevelSequence)
-		{
-			LevelSequence->GetMovieScene()->SetReadOnly(false);
-		}
-
-		if (UTakeMetaData*  TakeMetaDataLevelSequence = LevelSequence ? LevelSequence->FindMetaData<UTakeMetaData>() : nullptr)
-		{
-			TakeMetaDataLevelSequence->Unlock();
-			TakeMetaDataLevelSequence->SetTimestamp(FDateTime(0));
-		}
-	}
-	else if (InData.RecordIntoSequence)
-	{
-		SetRecordIntoLevelSequence(InData.RecordIntoSequence);
-	}
-	else if (InData.SequenceToView)
-	{
-		SuppliedLevelSequence = InData.SequenceToView;
-	}
-
-	/** Clear the dirty flag since the preset was just initialized. */
-	TransientPreset->GetOutermost()->SetDirtyFlag(false);
-
-	bHasTargetSequenceBeenSet = true;
-	SetTickableTickType(ETickableTickType::Conditional);
+	Implementation->SetTargetSequence(InData);
 }
 
 void UTakeRecorderSubsystem::SetRecordIntoLevelSequence(ULevelSequence* LevelSequence)
 {
-	SuppliedLevelSequence = nullptr;
-	RecordIntoLevelSequence = LevelSequence;
-
-	if (RecordIntoLevelSequence)
-	{
-		RecordIntoLevelSequence->GetMovieScene()->SetReadOnly(false);
-	}
-
-	if (UTakeMetaData* LocalTakeMetaData = RecordIntoLevelSequence ? RecordIntoLevelSequence->FindOrAddMetaData<UTakeMetaData>() : nullptr)
-	{
-		// Set up take metadata to match this level sequence's info, ie. match the frame rate, use the level sequence name as the slate
-		LocalTakeMetaData->Unlock();
-		LocalTakeMetaData->SetTimestamp(FDateTime(0));
-		LocalTakeMetaData->SetSlate(LevelSequence->GetName());
-		LocalTakeMetaData->SetTakeNumber(0);
-		LocalTakeMetaData->SetFrameRate(LevelSequence->GetMovieScene()->GetDisplayRate());
-		LocalTakeMetaData->SetFrameRateFromTimecode(false);
-	}
+	Implementation->SetRecordIntoLevelSequence(LevelSequence);
 }
 
 bool UTakeRecorderSubsystem::CanReviewLastRecording() const
 {
-	if (GetLastRecordedLevelSequence() == nullptr)
-	{
-		return false;
-	}
-	ITakeRecorderModule& TakeRecorderModule = FModuleManager::Get().LoadModuleChecked<ITakeRecorderModule>("TakeRecorder");
-	FCanReviewLastRecordedLevelSequence& CanReview = TakeRecorderModule.GetCanReviewLastRecordedLevelSequenceDelegate();
-	if (CanReview.IsBound())
-	{
-		return CanReview.Execute();
-	}
-	
-	return true;
+	return Implementation->CanReviewLastRecording();
 }
 
 bool UTakeRecorderSubsystem::ReviewLastRecording()
 {
-	if (LastRecordedLevelSequence.IsValid())
-	{
-		SuppliedLevelSequence = LastRecordedLevelSequence.Get();
-		LastRecordedLevelSequence = nullptr;
-		return true;
-	}
-
-	return false;
+	return Implementation->ReviewLastRecording();
 }
 
 bool UTakeRecorderSubsystem::StartRecording(bool bOpenSequencer, bool bShowErrorMessage)
 {
-	static bool bStartedRecording = false;
-
-	if (bStartedRecording)
-	{
-		return false;
-	}
-
-	if (!GIsEditor)
-	{
-		// Sequencer requires an editor attached.
-		bOpenSequencer = false;
-	}
-
-	TGuardValue<bool> ReentrantGuard(bStartedRecording, true);
-
-	ULevelSequence*       LevelSequence = GetLevelSequence();
-	UTakeRecorderSources* Sources = LevelSequence ? LevelSequence->FindMetaData<UTakeRecorderSources>() : nullptr;
-	
-	if (LevelSequence && Sources)
-	{
-		FTakeRecorderParameters Parameters;
-		Parameters.User    = GetDefault<UTakeRecorderUserSettings>()->Settings;
-		Parameters.Project = GetDefault<UTakeRecorderProjectSettings>()->Settings;
-		Parameters.TakeRecorderMode = GetTakeRecorderMode();
-		Parameters.StartFrame = LevelSequence->GetMovieScene()->GetPlaybackRange().GetLowerBoundValue();
-		Parameters.bOpenSequencer = bOpenSequencer;
-
-		FText ErrorText = LOCTEXT("UnknownError", "An unknown error occurred when trying to start recording");
-
-		if (GIsEditor)
-		{
-			IAssetEditorInstance* AssetEditor = GEditor->GetEditorSubsystem<UAssetEditorSubsystem>()->FindEditorForAsset(LevelSequence, false);
-			ILevelSequenceEditorToolkit* LevelSequenceEditor = static_cast<ILevelSequenceEditorToolkit*>(AssetEditor);
-
-			if (LevelSequenceEditor && LevelSequenceEditor->GetSequencer())
-			{
-				// If not resetting the playhead, store the current time as the start frame for recording. 
-				// This will ultimately be the start of the playback range and the recording will begin from that time.
-				if (!Parameters.User.bResetPlayhead)
-				{
-					Parameters.StartFrame = LevelSequenceEditor->GetSequencer()->GetLocalTime().Time.FrameNumber;
-				}
-			}
-		}
-
-		UTakeRecorder* NewRecorder = NewObject<UTakeRecorder>(GetTransientPackage(), NAME_None, RF_Transient);
-
-		if (!NewRecorder->Initialize(LevelSequence, Sources, TakeMetaData, Parameters, &ErrorText))
-		{
-			if (bShowErrorMessage && ensure(!ErrorText.IsEmpty()))
-			{
-				FNotificationInfo Info(ErrorText);
-				Info.ExpireDuration = 5.0f;
-				FSlateNotificationManager::Get().AddNotification(Info)->SetCompletionState(SNotificationItem::CS_Fail);
-			}
-			return false;
-		}
-
-		OnRecordingStarted(NewRecorder);
-		
-		return true;
-	}
-	
-	return false;
+	return Implementation->StartRecording(bOpenSequencer, bShowErrorMessage);
 }
 
 void UTakeRecorderSubsystem::StopRecording()
 {
-	if (UTakeRecorder* CurrentRecording = UTakeRecorder::GetActiveRecorder())
-	{
-		CurrentRecording->Stop();
-	}
+	Implementation->StopRecording();
 }
 
 void UTakeRecorderSubsystem::CancelRecording(bool bShowConfirmMessage)
 {
-	if (UTakeRecorder* CurrentRecording = UTakeRecorder::GetActiveRecorder())
-	{
-		auto DoCancel = [](const TWeakObjectPtr<UTakeRecorder>& WeakRecording)
-		{
-			if (WeakRecording.IsValid())
-			{
-				WeakRecording.Get()->Cancel(); 
-			}
-		};
-
-		if (bShowConfirmMessage)
-		{
-			const TSharedRef<SMessageDialog> ConfirmDialog = SNew(SMessageDialog)
-				.Title(FText(LOCTEXT("ConfirmCancelRecordingTitle", "Cancel Recording?")))
-				.Message(LOCTEXT("ConfirmCancelRecording", "Are you sure you want to cancel the current recording?"))
-				.Buttons({
-					SCustomDialog::FButton(LOCTEXT("Yes", "Yes"), FSimpleDelegate::CreateLambda([DoCancel, WeakRecording = TWeakObjectPtr<UTakeRecorder>(CurrentRecording)]()
-						{ 
-							DoCancel(WeakRecording);
-						})),
-					SCustomDialog::FButton(LOCTEXT("No", "No"))
-				});
-
-			// Non modal so that the recording continues to update
-			ConfirmDialog->Show();
-		}
-		else
-		{
-			DoCancel(TWeakObjectPtr<UTakeRecorder>(CurrentRecording));
-		}
-	}
+	Implementation->CancelRecording(bShowConfirmMessage);
 }
 
 void UTakeRecorderSubsystem::ResetToPendingTake()
 {
-	if (IsReviewing())
-	{
-		LastRecordedLevelSequence = SuppliedLevelSequence;
-	}
-
-	SuppliedLevelSequence = nullptr;
-	RecordIntoLevelSequence = nullptr;
-	
-	TransientPreset = AllocateTransientPreset();
+	Implementation->ResetToPendingTake();
 }
 
 void UTakeRecorderSubsystem::ClearPendingTake()
 {
-	if (IsReviewing())
-	{
-		LastRecordedLevelSequence = SuppliedLevelSequence;
-	}
-
-	UTakeRecorderSources* BaseSources = nullptr;
-
-	if (const ULevelSequence* CurrentLevelSequence = GetLevelSequence())
-	{
-		BaseSources = CurrentLevelSequence->FindMetaData<UTakeRecorderSources>();
-	}
-
-	SuppliedLevelSequence = nullptr;
-	RecordIntoLevelSequence = nullptr;
-
-	FScopedTransaction Transaction(LOCTEXT("ClearPendingTake_Transaction", "Clear Pending Take"));
-
-	TransientPreset->Modify();
-	TransientPreset->CreateLevelSequence();
-
-	ULevelSequence* LevelSequence = TransientPreset->GetLevelSequence();
-	if (LevelSequence && BaseSources)
-	{
-		LevelSequence->CopyMetaData(BaseSources);
-	}
+	Implementation->ClearPendingTake();
 }
 
 UTakePreset* UTakeRecorderSubsystem::GetPendingTake() const
 {
-	const ITakeRecorderModule& TakeRecorderModule = FModuleManager::LoadModuleChecked<ITakeRecorderModule>("TakeRecorder");
-	return TakeRecorderModule.GetPendingTake();
+	return Implementation->GetPendingTake();
 }
 
 void UTakeRecorderSubsystem::RevertChanges()
 {
-	UTakePreset* PresetOrigin = GetTakeMetaData()->GetPresetOrigin();
-
-	FScopedTransaction Transaction(LOCTEXT("RevertChanges_Transaction", "Revert Changes"));
-
-	TransientPreset->Modify();
-	TransientPreset->CopyFrom(PresetOrigin);
-	TransientPreset->GetOutermost()->SetDirtyFlag(false);
+	Implementation->RevertChanges();
 }
 
 UTakeRecorderSource* UTakeRecorderSubsystem::AddSource(const TSubclassOf<UTakeRecorderSource> InSourceClass)
 {
-	UTakeRecorderSources* Sources = GetSources();
-
-	if (*InSourceClass && Sources)
-	{
-		FScopedTransaction Transaction(FText::Format(LOCTEXT("AddNewSource", "Add New {0} Source"), InSourceClass->GetDisplayNameText()));
-		Sources->Modify();
-
-		return Sources->AddSource(InSourceClass);
-	}
-
-	return nullptr;
+	return Implementation->AddSource(InSourceClass);
 }
 
 void UTakeRecorderSubsystem::RemoveSource(UTakeRecorderSource* InSource)
 {
-	UTakeRecorderSources* Sources = GetSources();
-	if (Sources && InSource)
-	{
-		Sources->RemoveSource(InSource);
-	}
+	Implementation->RemoveSource(InSource);
 }
 
 void UTakeRecorderSubsystem::ClearSources()
 {
-	if (UTakeRecorderSources* Sources = GetSources())
-	{
-		for (UTakeRecorderSource* Source : Sources->GetSourcesCopy())
-		{
-			Sources->RemoveSource(Source);
-		}
-	}
+	Implementation->ClearSources();
 }
 
 UTakeRecorderSources* UTakeRecorderSubsystem::GetSources() const
 {
-	ULevelSequence* LevelSequence = GetLevelSequence();
-	UTakeRecorderSources* Sources = LevelSequence ? LevelSequence->FindOrAddMetaData<UTakeRecorderSources>() : nullptr;
-	return Sources;
+	return Implementation->GetSources();
 }
 
 TArrayView<UTakeRecorderSource* const> UTakeRecorderSubsystem::GetAllSources() const
 {
-	if (const UTakeRecorderSources* Sources = GetSources())
-	{
-		return Sources->GetSources();
-	}
-
-	return {};
+	return Implementation->GetAllSources();
 }
 
 TArray<UTakeRecorderSource*> UTakeRecorderSubsystem::GetAllSourcesCopy() const
 {
-	return TArray<UTakeRecorderSource*>(GetAllSources());
+	return Implementation->GetAllSourcesCopy();
 }
 
 UTakeRecorderSource* UTakeRecorderSubsystem::GetSourceByClass(const TSubclassOf<UTakeRecorderSource> InSourceClass) const
 {
-	const TArrayView<UTakeRecorderSource* const> Sources = GetAllSources();
-	for (UTakeRecorderSource* Source : Sources)
-	{
-		if (Source->IsA(InSourceClass))
-		{
-			return Source;
-		}
-	}
-
-	return nullptr;
+	return Implementation->GetSourceByClass(InSourceClass);
 }
 
 void UTakeRecorderSubsystem::AddSourceForActor(AActor* InActor, bool bReduceKeys, bool bShowProgress)
 {
-	if (UTakeRecorderSources* Sources = GetSources())
-	{
-		FTakeRecorderModule& TakeRecorderModule = FModuleManager::GetModuleChecked<FTakeRecorderModule>("TakeRecorder");
-		TakeRecorderModule.GetSourcesExtensionData().OnAddActorSource.ExecuteIfBound(Sources, { InActor }, bReduceKeys, bShowProgress);
-	}
+	Implementation->AddSourceForActor(InActor, bReduceKeys, bShowProgress);
 }
 
 void UTakeRecorderSubsystem::RemoveActorFromSources(AActor* InActor)
 {
-	if (UTakeRecorderSources* Sources = GetSources())
-	{
-		FTakeRecorderModule& TakeRecorderModule = FModuleManager::GetModuleChecked<FTakeRecorderModule>("TakeRecorder");
-		TakeRecorderModule.GetSourcesExtensionData().OnRemoveActorSource.ExecuteIfBound(Sources, { InActor });
-	}
+	Implementation->RemoveActorFromSources(InActor);
 }
 
 AActor* UTakeRecorderSubsystem::GetSourceActor(UTakeRecorderSource* InSource) const
 {
-	FTakeRecorderModule& TakeRecorderModule = FModuleManager::GetModuleChecked<FTakeRecorderModule>("TakeRecorder");
-	if (TakeRecorderModule.GetSourcesExtensionData().OnGetSourceActor.IsBound())
-	{
-		return TakeRecorderModule.GetSourcesExtensionData().OnGetSourceActor.Execute(InSource);
-	}
-
-	return nullptr;
+	return Implementation->GetSourceActor(InSource);
 }
 
 ETakeRecorderState UTakeRecorderSubsystem::GetState() const
 {
-	if (const UTakeRecorder* Recorder = UTakeRecorder::GetActiveRecorder())
-	{
-		return Recorder->GetState();
-	}
-
-	return ETakeRecorderState::PreInitialization;
+	return Implementation->GetState();
 }
 
 void UTakeRecorderSubsystem::SetTakeNumber(int32 InNewTakeNumber, bool bEmitChanged)
 {
-	if (TakeMetaData)
-	{
-		FScopedTransaction Transaction(LOCTEXT("SetTake_Transaction", "Set Take Number"));
-		TakeMetaData->Modify();
-		TakeMetaData->SetTakeNumber(InNewTakeNumber, bEmitChanged);
-		bAutoApplyTakeNumber = false;
-	}
+	Implementation->SetTakeNumber(InNewTakeNumber, bEmitChanged);
 }
 
 int32 UTakeRecorderSubsystem::GetNextTakeNumber(const FString& InSlate) const
 {
-	return UTakesCoreBlueprintLibrary::ComputeNextTakeNumber(InSlate);
+	return Implementation->GetNextTakeNumber(InSlate);
 }
 
 void UTakeRecorderSubsystem::GetNumberOfTakes(const FString& InSlate, int32& OutMaxTake, int32& OutNumTakes) const
 {
-	int32 MaxTake = 0;
-
-	TArray<FAssetData> Takes = UTakesCoreBlueprintLibrary::FindTakes(InSlate);
-
-	for (const FAssetData& Asset : Takes)
-	{
-		FAssetDataTagMapSharedView::FFindTagResult TakeNumberTag = Asset.TagsAndValues.FindTag(UTakeMetaData::AssetRegistryTag_TakeNumber);
-
-		int32 ThisTakeNumber = 0;
-		if (TakeNumberTag.IsSet() && LexTryParseString(ThisTakeNumber, *TakeNumberTag.GetValue()))
-		{
-			MaxTake = FMath::Max(MaxTake, ThisTakeNumber);
-		}
-	}
-
-	OutMaxTake = MaxTake;
-	OutNumTakes = Takes.Num();
+	Implementation->GetNumberOfTakes(InSlate, OutMaxTake, OutNumTakes);
 }
 
 TArray<FAssetData> UTakeRecorderSubsystem::GetSlates(FName InPackagePath) const
 {
-	const IAssetRegistry& AssetRegistry = FModuleManager::LoadModuleChecked<FAssetRegistryModule>(AssetRegistryConstants::ModuleName).Get();
-
-	FARFilter Filter;
-	Filter.ClassPaths.Add(ULevelSequence::StaticClass()->GetClassPathName());
-	Filter.TagsAndValues.Add(UTakeMetaData::AssetRegistryTag_Slate);
-	if (!InPackagePath.IsNone())
-	{
-		Filter.PackagePaths.Add(InPackagePath);
-		Filter.bRecursivePaths = true;
-	}
-
-	TArray<FAssetData> AllAssets;
-	AssetRegistry.GetAssets(Filter, AllAssets);
-	
-	return AllAssets;
+	return Implementation->GetSlates(InPackagePath);
 }
 
 void UTakeRecorderSubsystem::SetSlateName(const FString& InSlateName, bool bEmitChanged)
 {
-	if (TakeMetaData && TakeMetaData->GetSlate() != InSlateName)
-	{
-		FScopedTransaction Transaction(LOCTEXT("SetSlate_Transaction", "Set Take Slate"));
-		TakeMetaData->Modify();
-
-		TakeMetaData->SetSlate(InSlateName, bEmitChanged);
-
-		// Compute the correct starting take number
-		const int32 NextTakeNumber = GetNextTakeNumber(TakeMetaData->GetSlate());
-		if (NextTakeNumber != TakeMetaData->GetTakeNumber())
-		{
-			TakeMetaData->SetTakeNumber(NextTakeNumber, bEmitChanged);
-		}
-	}
+	Implementation->SetSlateName(InSlateName, bEmitChanged);
 }
 
 bool UTakeRecorderSubsystem::MarkFrame()
 {
-	if (UTakeRecorderBlueprintLibrary::IsRecording() && TakeMetaData)
-	{
-		const FFrameRate FrameRate = TakeMetaData->GetFrameRate();
-
-		const FTimespan RecordingDuration = FDateTime::UtcNow() - TakeMetaData->GetTimestamp();
-
-		const FFrameNumber ElapsedFrame = FFrameNumber(static_cast<int32>(FrameRate.AsDecimal() * RecordingDuration.GetTotalSeconds()));
-
-		const ULevelSequence* LevelSequence = GetLevelSequence();
-		if (!LevelSequence)
-		{
-			return false;
-		}
-
-		UMovieScene* MovieScene = LevelSequence->GetMovieScene();
-		if (!MovieScene)
-		{
-			return false;
-		}
-
-		const FFrameRate DisplayRate = MovieScene->GetDisplayRate();
-		const FFrameRate TickResolution = MovieScene->GetTickResolution();
-
-		FMovieSceneMarkedFrame MarkedFrame;
-
-		const UTakeRecorderSources* Sources = LevelSequence->FindMetaData<UTakeRecorderSources>();
-		if (Sources && Sources->GetSettings().bStartAtCurrentTimecode)
-		{
-			MarkedFrame.FrameNumber = FFrameRate::TransformTime(FFrameTime(FApp::GetTimecode().ToFrameNumber(DisplayRate)), DisplayRate, TickResolution).FloorToFrame();
-		}
-		else
-		{
-			MarkedFrame.FrameNumber = ConvertFrameTime(ElapsedFrame, DisplayRate, TickResolution).CeilToFrame();
-		}
-
-		const int32 MarkedFrameIndex = MovieScene->AddMarkedFrame(MarkedFrame);
-		UTakeRecorderBlueprintLibrary::OnTakeRecorderMarkedFrameAdded(MovieScene->GetMarkedFrames()[MarkedFrameIndex]);
-		if (const UTakeRecorderSubsystem* TakeRecorderSubsystem = GEngine->GetEngineSubsystem<UTakeRecorderSubsystem>())
-		{
-			TakeRecorderSubsystem->TakeRecorderMarkedFrameAdded.Broadcast(MovieScene->GetMarkedFrames()[MarkedFrameIndex]);
-		}
-
-		return true;
-	}
-
-	return false;
+	return Implementation->MarkFrame();
 }
 
 FFrameRate UTakeRecorderSubsystem::GetFrameRate() const
 {
-	if (TakeMetaData)
-	{
-		return TakeMetaData->GetFrameRate();
-	}
-
-	return FFrameRate();
+	return Implementation->GetFrameRate();
 }
 
 void UTakeRecorderSubsystem::SetFrameRate(FFrameRate InFrameRate)
 {
-	SetFrameRateImpl(InFrameRate, false);
+	Implementation->SetFrameRate(InFrameRate);
 }
 
 void UTakeRecorderSubsystem::SetFrameRateFromTimecode()
 {
-	SetFrameRateImpl(FApp::GetTimecodeFrameRate(), true);
+	Implementation->SetFrameRateFromTimecode();
 }
 
 void UTakeRecorderSubsystem::ImportPreset(const FAssetData& InPreset)
 {
-	SuppliedLevelSequence = nullptr;
-	RecordIntoLevelSequence = nullptr;
-
-	UTakePreset* Take = CastChecked<UTakePreset>(InPreset.GetAsset());
-	FScopedTransaction Transaction(LOCTEXT("ImportPreset_Transaction", "Import Take Preset"));
-
-	TransientPreset->Modify();
-	TransientPreset->CopyFrom(Take);
-	TransientPreset->GetOutermost()->SetDirtyFlag(false);
-
-	GetTakeMetaData()->SetPresetOrigin(Take);
+	Implementation->ImportPreset(InPreset);
 }
 
 bool UTakeRecorderSubsystem::IsReviewing() const
 {
-	return !IsRecording() && (TakeMetaData && TakeMetaData->Recorded() && GetTakeRecorderMode() != ETakeRecorderMode::RecordIntoSequence);
+	return Implementation->IsReviewing();
 }
 
 bool UTakeRecorderSubsystem::IsRecording() const
 {
-	return UTakeRecorderBlueprintLibrary::GetActiveRecorder() ? true : false;
+	return Implementation->IsRecording();
 }
 
 bool UTakeRecorderSubsystem::TryGetSequenceCountdown(float& OutValue) const
 {
-	const UTakeRecorder* CurrentRecording = UTakeRecorder::GetActiveRecorder();
-	const bool           bIsCountingDown  = CurrentRecording && CurrentRecording->GetState() == ETakeRecorderState::CountingDown;
-	
-	OutValue = bIsCountingDown ? CurrentRecording->GetCountdownSeconds() : 0.f;
-	return bIsCountingDown;
+	return Implementation->TryGetSequenceCountdown(OutValue);
 }
 
 void UTakeRecorderSubsystem::SetSequenceCountdown(float InSeconds)
 {
-	if (UTakeRecorder* CurrentRecording = UTakeRecorder::GetActiveRecorder())
-	{
-		CurrentRecording->SetCountdown(InSeconds);
-	}
+	Implementation->SetSequenceCountdown(InSeconds);
 }
 
 TArray<UObject*> UTakeRecorderSubsystem::GetSourceRecordSettings(UTakeRecorderSource* InSource) const
 {
-	return InSource->GetAdditionalSettingsObjects();
+	return Implementation->GetSourceRecordSettings(InSource);
 }
 
 FTakeRecorderParameters UTakeRecorderSubsystem::GetGlobalRecordSettings() const
 {
-	return UTakeRecorderBlueprintLibrary::GetDefaultParameters();
+	return Implementation->GetGlobalRecordSettings();
 }
 
 void UTakeRecorderSubsystem::SetGlobalRecordSettings(const FTakeRecorderParameters& InParameters)
 {
-	UTakeRecorderBlueprintLibrary::SetDefaultParameters(InParameters);
+	Implementation->SetGlobalRecordSettings(InParameters);
 }
 
 UTakeMetaData* UTakeRecorderSubsystem::GetTakeMetaData() const
 {
-	return TakeMetaData;
+	return Implementation->GetTakeMetaData();
 }
 
 ULevelSequence* UTakeRecorderSubsystem::GetLevelSequence() const
 {
-	if (SuppliedLevelSequence)
-	{
-		return SuppliedLevelSequence;
-	}
-	else if (RecordIntoLevelSequence)
-	{
-		return RecordIntoLevelSequence;
-	}
-	else if (RecordingLevelSequence)
-	{
-		return RecordingLevelSequence;
-	}
-	else if (TransientPreset)
-	{
-		return TransientPreset->GetLevelSequence();
-	}
-	return nullptr;
+	return Implementation->GetLevelSequence();
 }
 
 ULevelSequence* UTakeRecorderSubsystem::GetSuppliedLevelSequence() const
 {
-	return SuppliedLevelSequence;
+	return Implementation->GetSuppliedLevelSequence();
 }
 
 ULevelSequence* UTakeRecorderSubsystem::GetRecordingLevelSequence() const
 {
-	return RecordingLevelSequence;
+	return Implementation->GetRecordingLevelSequence();
 }
 
 ULevelSequence* UTakeRecorderSubsystem::GetRecordIntoLevelSequence() const
 {
-	return RecordIntoLevelSequence;
+	return Implementation->GetRecordIntoLevelSequence();
 }
 
 ULevelSequence* UTakeRecorderSubsystem::GetLastRecordedLevelSequence() const
 {
-	return LastRecordedLevelSequence.Get();
+	return Implementation->GetLastRecordedLevelSequence();
 }
 
 UTakePreset* UTakeRecorderSubsystem::GetTransientPreset() const
 {
-	return TransientPreset;
+	return Implementation->GetTransientPreset();
 }
 
 ETakeRecorderMode UTakeRecorderSubsystem::GetTakeRecorderMode() const
 {
-	if (RecordIntoLevelSequence != nullptr)
-	{
-		return ETakeRecorderMode::RecordIntoSequence;
-	}
-
-	return ETakeRecorderMode::RecordNewSequence;
+	return Implementation->GetTakeRecorderMode();
 }
 
 UTakeRecorderNamingTokensData* UTakeRecorderSubsystem::GetNamingTokensData() const
 {
-	check(NamingTokensData);
-	return NamingTokensData;
+	return Implementation->GetNamingTokensData();
 }
 
 bool UTakeRecorderSubsystem::HasPendingChanges() const
 {
-	UMovieScene* MovieScene = GetLevelSequence() ? GetLevelSequence()->GetMovieScene() : nullptr;
-	const bool bHasUserMadeChanges = MovieScene
-		&& (!MovieScene->GetTracks().IsEmpty()
-			|| !MovieScene->GetBindings().IsEmpty()
-			|| MovieScene->GetPossessableCount() > 0
-			|| MovieScene->GetSpawnableCount() > 0);
-	return bHasUserMadeChanges;
+	return Implementation->HasPendingChanges();
 }
 
 FOnTakeRecordingInitialized& UTakeRecorderSubsystem::GetOnRecordingPreInitializedEvent()
@@ -862,296 +313,4 @@ UTakeRecorderSources::FOnSourceRemoved& UTakeRecorderSubsystem::GetOnRecordingSo
 	return OnRecordingSourceRemovedEvent;
 }
 
-UTakePreset* UTakeRecorderSubsystem::AllocateTransientPreset()
-{
-	return UTakePreset::AllocateTransientPreset(GetDefault<UTakeRecorderUserSettings>()->LastOpenedPreset.Get());
-}
-
-void UTakeRecorderSubsystem::CacheMetaData()
-{
-	UTakeMetaData* NewMetaDataThisTick = nullptr;
-
-	if (const ULevelSequence* LevelSequence = GetLevelSequence())
-	{
-		NewMetaDataThisTick = LevelSequence->FindMetaData<UTakeMetaData>();
-	}
-
-	// If it's null we use the transient meta-data
-	if (!NewMetaDataThisTick)
-	{
-		// if the transient meta-data doesn't exist, create it now
-		if (!TransientTakeMetaData)
-		{
-			TransientTakeMetaData = UTakeMetaData::CreateFromDefaults(GetTransientPackage(), NAME_None);
-			TransientTakeMetaData->SetFlags(RF_Transactional | RF_Transient);
-
-			FString DefaultSlate = GetDefault<UTakeRecorderProjectSettings>()->Settings.DefaultSlate;
-			if (TransientTakeMetaData->GetSlate() != DefaultSlate)
-			{
-				TransientTakeMetaData->SetSlate(DefaultSlate, false);
-			}
-
-			// Compute the correct starting take number
-			int32 NextTakeNumber = GetNextTakeNumber(TransientTakeMetaData->GetSlate());
-			if (TransientTakeMetaData->GetTakeNumber() != NextTakeNumber)
-			{
-				TransientTakeMetaData->SetTakeNumber(NextTakeNumber, false);
-			}
-
-			UTakeMetaData::SetMostRecentMetaData(TransientTakeMetaData);
-		}
-
-		NewMetaDataThisTick = TransientTakeMetaData;
-	}
-
-	check(NewMetaDataThisTick);
-	if (NewMetaDataThisTick != TakeMetaData)
-	{
-		TakeMetaData = NewMetaDataThisTick;
-
-		if (!OnPreForceDeleteObjectsHandle.IsValid())
-		{
-			// Since this is a strong reference, we need to be able to clear it if the user is attempting to remove it, but let's not always
-			// hook this event since the subsystem won't be active for general UE operation.
-			OnPreForceDeleteObjectsHandle = FEditorDelegates::OnPreForceDeleteObjects.AddUObject(this, &UTakeRecorderSubsystem::OnPreForceDeleteObjects);
-		}
-	}
-
-	//Set MovieScene Display Rate to the Preset Frame Rate.
-	if (const ULevelSequence* Sequence = GetLevelSequence())
-	{
-		if (UMovieScene* MovieScene = Sequence->GetMovieScene())
-		{
-			MovieScene->SetDisplayRate(TakeMetaData->GetFrameRate());
-		}
-	}
-
-	check(TakeMetaData);
-}
-
-void UTakeRecorderSubsystem::SetFrameRateImpl(const FFrameRate& InFrameRate, bool bFromTimecode)
-{
-	if (TakeMetaData)
-	{
-		TakeMetaData->SetFrameRateFromTimecode(bFromTimecode);
-		TakeMetaData->SetFrameRate(InFrameRate);
-	}
-	const ULevelSequence* Sequence = GetLevelSequence();
-	if (UMovieScene* MovieScene = Sequence ? Sequence->GetMovieScene() : nullptr)
-	{
-		MovieScene->SetDisplayRate(InFrameRate);
-	}
-}
-
-void UTakeRecorderSubsystem::IncrementTakeNumber()
-{
-	if (TransientTakeMetaData)
-	{
-		// Increment the transient take meta data if necessary
-		int32 NextTakeNumber = GetNextTakeNumber(TransientTakeMetaData->GetSlate());
-
-		if (TransientTakeMetaData->GetTakeNumber() != NextTakeNumber)
-		{
-			TransientTakeMetaData->SetTakeNumber(NextTakeNumber);
-		}
-	}
-	
-	// Update the preset take number at the end of recording
-	
-	if (const ULevelSequence* LevelSequence = TransientPreset ? TransientPreset->GetLevelSequence() : nullptr)
-	{
-		if (UTakeMetaData* MetaData = LevelSequence->FindMetaData<UTakeMetaData>())
-		{
-			const int32 NextTakeNumber = GetNextTakeNumber(MetaData->GetSlate());
-			MetaData->SetTakeNumber(NextTakeNumber);
-		}
-	}
-
-	bAutoApplyTakeNumber = true;
-}
-
-void UTakeRecorderSubsystem::OnAssetRegistryFilesLoaded()
-{
-	if (bAutoApplyTakeNumber && TransientTakeMetaData)
-	{
-		const int32 NextTakeNumber = GetNextTakeNumber(TransientTakeMetaData->GetSlate());
-		TransientTakeMetaData->SetTakeNumber(NextTakeNumber);
-	}
-}
-
-void UTakeRecorderSubsystem::OnRecordingInitialized(UTakeRecorder* Recorder)
-{
-	if (!bHasTargetSequenceBeenSet)
-	{
-		// Not initialized, take recorder triggered outside of the subsystem.
-		SetTargetSequence();
-	}
-	
-	// This needs to be stored with a strong ptr before our panels refresh, otherwise the weak sequencer ref will be lost during initialize.
-	RecordingLevelSequence = Recorder->GetSequence();
-	// Recache the meta-data here since we know that the sequence has probably changed as a result of the recording being started
-	CacheMetaData();
-	OnRecordingFinishedHandle = Recorder->OnRecordingFinished().AddUObject(this, &UTakeRecorderSubsystem::OnRecordingFinished);
-	OnRecordingCancelledHandle = Recorder->OnRecordingCancelled().AddUObject(this, &UTakeRecorderSubsystem::OnRecordingCancelled);
-
-	GetOnRecordingInitializedEvent().Broadcast(Recorder);
-	TakeRecorderInitialized.Broadcast();
-}
-
-void UTakeRecorderSubsystem::OnRecordingStarted(UTakeRecorder* Recorder)
-{
-	GetOnRecordingStartedEvent().Broadcast(Recorder);
-	TakeRecorderStarted.Broadcast();
-}
-
-void UTakeRecorderSubsystem::OnRecordingStopped(UTakeRecorder* Recorder)
-{
-	Recorder->OnRecordingStopped().Remove(OnRecordingStoppedHandle);
-	GetOnRecordingStoppedEvent().Broadcast(Recorder);
-	TakeRecorderStopped.Broadcast();
-}
-
-void UTakeRecorderSubsystem::OnRecordingFinished(UTakeRecorder* Recorder)
-{
-	LastRecordedLevelSequence = RecordingLevelSequence;
-	RecordingLevelSequence = nullptr;
-
-	IncrementTakeNumber();
-	
-	if (TransientTakeMetaData)
-	{
-		// Increment the transient take meta data if necessary
-		int32 NextTakeNumber = GetNextTakeNumber(TransientTakeMetaData->GetSlate());
-
-		if (TransientTakeMetaData->GetTakeNumber() != NextTakeNumber)
-		{
-			TransientTakeMetaData->SetTakeNumber(NextTakeNumber);
-		}
-
-		bAutoApplyTakeNumber = true;
-	}
-
-	Recorder->OnRecordingFinished().Remove(OnRecordingFinishedHandle);
-
-	GetOnRecordingFinishedEvent().Broadcast(Recorder);
-	TakeRecorderFinished.Broadcast(LastRecordedLevelSequence.Get());
-}
-
-void UTakeRecorderSubsystem::OnRecordingCancelled(UTakeRecorder* Recorder)
-{
-	RecordingLevelSequence = nullptr;
-
-	Recorder->OnRecordingFinished().Remove(OnRecordingFinishedHandle);
-	Recorder->OnRecordingCancelled().Remove(OnRecordingCancelledHandle);
-
-	GetOnRecordingCancelledEvent().Broadcast(Recorder);
-	TakeRecorderCancelled.Broadcast();
-}
-
-void UTakeRecorderSubsystem::OnTakeSlateChanged(const FString& InSlate, UTakeMetaData* InTakeMetaData)
-{
-	TakeRecorderSlateChanged.Broadcast(InSlate, InTakeMetaData);
-}
-
-void UTakeRecorderSubsystem::OnTakeNumberChanged(int32 InTakeNumber, UTakeMetaData* InTakeMetaData)
-{
-	TakeRecorderTakeNumberChanged.Broadcast(InTakeNumber, InTakeMetaData);
-}
-
-void UTakeRecorderSubsystem::SetLastLevelSequence(ULevelSequence* InSequence)
-{
-	LastRecordedLevelSequence = InSequence;
-}
-
-void UTakeRecorderSubsystem::OnSourceAdded(UTakeRecorderSource* Source)
-{
-	GetOnRecordingSourceAddedEvent().Broadcast(Source);
-	TakeRecorderSourceAdded.Broadcast(Source);
-}
-
-void UTakeRecorderSubsystem::OnSourceRemoved(UTakeRecorderSource* Source)
-{
-	GetOnRecordingSourceRemovedEvent().Broadcast(Source);
-	TakeRecorderSourceRemoved.Broadcast(Source);
-}
-
-void UTakeRecorderSubsystem::BindToNamingTokenEvents()
-{
-	if (GEngine && !GetNamingTokensData()->TakeRecorderNamingTokens.IsValid())
-	{
-		GetNamingTokensData()->TakeRecorderNamingTokens = GEngine->GetEngineSubsystem<UNamingTokensEngineSubsystem>()->GetNamingTokens(ITakeRecorderNamingTokensModule::GetTakeRecorderNamespace());
-		if (ensure(GetNamingTokensData()->TakeRecorderNamingTokens.IsValid()))
-		{
-			const TStrongObjectPtr<UNamingTokens> TakeRecorderNamingTokensPin = GetNamingTokensData()->TakeRecorderNamingTokens.Pin();
-			TakeRecorderNamingTokensPin->RegisterExternalTokens(GetNamingTokensData()->NamingTokensExternalGuid);
-			TakeRecorderNamingTokensPin->GetOnPreEvaluateEvent().AddUObject(this, &UTakeRecorderSubsystem::OnTakeRecorderNamingTokensPreEvaluate);
-		}
-	}
-}
-
-void UTakeRecorderSubsystem::UnbindNamingTokensEvents()
-{
-	if (const TStrongObjectPtr<UNamingTokens> TakeRecorderNamingTokensPin = GetNamingTokensData()->TakeRecorderNamingTokens.Pin())
-	{
-		TakeRecorderNamingTokensPin->UnregisterExternalTokens(GetNamingTokensData()->NamingTokensExternalGuid);
-		TakeRecorderNamingTokensPin->GetOnPreEvaluateEvent().RemoveAll(this);
-	}
-}
-
-void UTakeRecorderSubsystem::OnTakeRecorderNamingTokensPreEvaluate(const FNamingTokensEvaluationData& InEvaluationData)
-{
-	if (GetNamingTokensData()->TakeRecorderNamingTokens.IsValid())
-	{
-		TArray<FNamingTokenData>& ExternalTokens = GetNamingTokensData()->TakeRecorderNamingTokens->GetExternalTokensChecked(GetNamingTokensData()->NamingTokensExternalGuid);
-
-		ExternalTokens.Reset(GetNamingTokensData()->UserDefinedTokens.Num());
-		for (const TTuple<FNamingTokenData, FText>& UserToken : GetNamingTokensData()->UserDefinedTokens)
-		{
-			ExternalTokens.Add({
-				UserToken.Key.TokenKey,
-				UserToken.Key.DisplayName,
-				// @todo naming tokens - If we support TOptional<FText> that can be used instead of a lambda
-				FNamingTokenData::FTokenProcessorDelegateNative::CreateLambda([Value = UserToken.Value]() {
-					return Value;
-				})
-			});
-		}
-	}
-}
-
-void UTakeRecorderSubsystem::OnTakePresetSettingsChanged()
-{
-	// This effectively regenerates the level sequence in the preset.
-	// Existing data is discarded. In the future, we should try to migrate it.
-	// FTakePresetRecorderCustomization handles asking the user whether it is ok to discard the changes. Once this event fires, the decision has been made.
-	if (!GIsTransacting)
-	{
-		ClearPendingTake();
-	}
-}
-
-void UTakeRecorderSubsystem::OnPreForceDeleteObjects(const TArray<UObject*>& InObjects)
-{
-	// Take meta data is cached as a strong reference, but can prevent us from deleting owning level sequences.
-	// We need to keep a strong reference to it for Take Recorder functionality, especially in the panel, and it should persist
-	// even with the panel closed. Listening for a force delete allows us to safely clear the reference.
-	const ULevelSequence* LevelSequence = GetLevelSequence();
-	if (!LevelSequence && LastRecordedLevelSequence.IsValid())
-	{
-		LevelSequence = LastRecordedLevelSequence.Get();
-	}
-	if (LevelSequence && TakeMetaData && InObjects.Contains(LevelSequence)
-		&& LevelSequence->FindMetaData<UTakeMetaData>() == TakeMetaData)
-	{
-		if (IsRecording())
-		{
-			// Need to cancel or we will be in a bad state. TakeRecorder will tick but be without a level sequence.
-			CancelRecording(false);
-		}
-		TakeMetaData = nullptr;
-		FEditorDelegates::OnPreForceDeleteObjects.Remove(OnPreForceDeleteObjectsHandle);
-		OnPreForceDeleteObjectsHandle.Reset();
-	}
-}
-
 #undef LOCTEXT_NAMESPACE
diff --git a/Engine/Plugins/VirtualProduction/Takes/Source/TakeRecorder/Private/Recorder/TakeRecorderSubsystemImplementation.cpp b/Engine/Plugins/VirtualProduction/Takes/Source/TakeRecorder/Private/Recorder/TakeRecorderSubsystemImplementation.cpp
new file mode 100644
--- /dev/null
+++ b/Engine/Plugins/VirtualProduction/Takes/Source/TakeRecorder/Private/Recorder/TakeRecorderSubsystemImplementation.cpp
@@ -0,0 +1,1129 @@
+// Copyright Epic Games, Inc. All Rights Reserved.
+
+#include "TakeRecorderSubsystemImplementation.h"
+
+#include "AssetRegistry/AssetRegistryModule.h"
+#include "Dialog/SMessageDialog.h"
+#include "Framework/Notifications/NotificationManager.h"
+#include "ILevelSequenceEditorToolkit.h"
+#include "ITakeRecorderModule.h"
+#include "ITakeRecorderNamingTokensModule.h"
+#include "LevelSequence.h"
+#include "MovieScene.h"
+#include "NamingTokens.h"
+#include "NamingTokensEngineSubsystem.h"
+#include "Recorder/TakeRecorder.h"
+#include "Recorder/TakeRecorderBlueprintLibrary.h"
+#include "Recorder/TakeRecorderParameters.h"
+#include "Recorder/TakeRecorderSubsystem.h"
+#include "ScopedTransaction.h"
+#include "TakeMetaData.h"
+#include "TakePresetSettings.h"
+#include "TakeRecorderModule.h"
+#include "TakeRecorderSettings.h"
+#include "TakeRecorderSource.h"
+#include "TakeRecorderSources.h"
+#include "TakesCoreBlueprintLibrary.h"
+#include "Widgets/Notifications/SNotificationList.h"
+#include "Widgets/TakeRecorderNamingTokensData.h"
+
+#define LOCTEXT_NAMESPACE "TakeRecorderSubsystemImplementation"
+
+#define WITH_OWNING_SUBSYSTEM(CALL)                                               \
+do {                                                                              \
+	if (UTakeRecorderSubsystem* OwningSubsystem = OwningSubsystemWeakPtr.Get()) { \
+		OwningSubsystem->CALL;                                                    \
+	}                                                                             \
+} while (0)
+
+ETickableTickType UTakeRecorderSubsystemImplementation::GetTickableTickType() const
+{
+	// This is to prevent registration until we call SetTargetSequence.
+	return ETickableTickType::Never;
+}
+
+UWorld* UTakeRecorderSubsystemImplementation::GetTickableGameObjectWorld() const
+{
+	return FTickableGameObject::GetTickableGameObjectWorld();
+}
+
+bool UTakeRecorderSubsystemImplementation::IsTickable() const
+{
+	return bHasTargetSequenceBeenSet;
+}
+
+bool UTakeRecorderSubsystemImplementation::IsTickableWhenPaused() const
+{
+	return FTickableGameObject::IsTickableWhenPaused();
+}
+
+bool UTakeRecorderSubsystemImplementation::IsTickableInEditor() const
+{
+	return true;
+}
+
+void UTakeRecorderSubsystemImplementation::Tick(float DeltaTime)
+{
+	CacheMetaData();
+}
+
+TStatId UTakeRecorderSubsystemImplementation::GetStatId() const
+{
+	RETURN_QUICK_DECLARE_CYCLE_STAT(UTakeRecorderSubsystemImplementation, STATGROUP_Tickables);
+}
+
+void UTakeRecorderSubsystemImplementation::InitializeImplementation(UTakeRecorderSubsystem* OwningSubsystem)
+{
+	check(OwningSubsystem);
+	OwningSubsystemWeakPtr = OwningSubsystem;
+	
+	const FAssetRegistryModule& AssetRegistryModule = FModuleManager::LoadModuleChecked<FAssetRegistryModule>(TEXT("AssetRegistry"));
+	OnAssetRegistryFilesLoadedHandle = AssetRegistryModule.Get().OnFilesLoaded().AddUObject(this, &UTakeRecorderSubsystemImplementation::OnAssetRegistryFilesLoaded);
+	OnRecordingInitializedHandle = UTakeRecorder::OnRecordingInitialized().AddUObject(this, &UTakeRecorderSubsystemImplementation::OnRecordingInitialized);
+
+	UTakeRecorderSources::OnSourceAdded().AddUObject(this, &UTakeRecorderSubsystemImplementation::OnSourceAdded);
+	UTakeRecorderSources::OnSourceRemoved().AddUObject(this, &UTakeRecorderSubsystemImplementation::OnSourceRemoved);
+
+	UTakeMetaData::OnTakeSlateChanged().AddUObject(this, &UTakeRecorderSubsystemImplementation::OnTakeSlateChanged);
+	UTakeMetaData::OnTakeNumberChanged().AddUObject(this, &UTakeRecorderSubsystemImplementation::OnTakeNumberChanged);
+
+	NamingTokensData = NewObject<UTakeRecorderNamingTokensData>();
+	NamingTokensData->SetFlags(RF_Transactional);
+
+	BindToNamingTokenEvents();
+	
+	// If the UTakePresetSettings::TargetRecordClass changes, the level sequence in the TransientPreset must be regenerated.
+	UTakePresetSettings::Get()->OnSettingsChanged().AddUObject(this, &UTakeRecorderSubsystemImplementation::OnTakePresetSettingsChanged);
+}
+
+void UTakeRecorderSubsystemImplementation::DeinitializeImplementation()
+{
+	UTakeRecorder::OnRecordingInitialized().Remove(OnRecordingInitializedHandle);
+
+	UTakeRecorderSources::OnSourceAdded().RemoveAll(this);
+	UTakeRecorderSources::OnSourceRemoved().RemoveAll(this);
+
+	UTakeMetaData::OnTakeSlateChanged().RemoveAll(this);
+	UTakeMetaData::OnTakeNumberChanged().RemoveAll(this);
+
+	UTakePresetSettings::Get()->OnSettingsChanged().RemoveAll(this);
+	
+	FEditorDelegates::OnPreForceDeleteObjects.Remove(OnPreForceDeleteObjectsHandle);
+	OnPreForceDeleteObjectsHandle.Reset();
+	
+	if (const FAssetRegistryModule* AssetRegistryModule = FModuleManager::GetModulePtr<FAssetRegistryModule>("AssetRegistry"))
+	{
+		if (IAssetRegistry* AssetRegistry = AssetRegistryModule->TryGet())
+		{
+			AssetRegistry->OnFilesLoaded().Remove(OnAssetRegistryFilesLoadedHandle);
+		}
+	}
+
+	if (ITakeRecorderModule* TakeRecorderModule = FModuleManager::Get().GetModulePtr<ITakeRecorderModule>("TakeRecorder"))
+	{
+		TakeRecorderModule->GetLastLevelSequenceProvider().Unbind();
+	}
+
+	UnbindNamingTokensEvents();
+
+	bHasTargetSequenceBeenSet = false;
+	SetTickableTickType(ETickableTickType::Never);
+
+	OwningSubsystemWeakPtr.Reset();
+}
+
+void UTakeRecorderSubsystemImplementation::SetTargetSequence(const FTakeRecorderSequenceParameters& InData)
+{
+	// Don't initialize again if we already have the same data present.
+	if (TargetSequenceData.IsValid() && TargetSequenceData == InData)
+	{
+		return;
+	}
+	TargetSequenceData = InData;
+	
+	ITakeRecorderModule& TakeRecorderModule = FModuleManager::Get().LoadModuleChecked<ITakeRecorderModule>("TakeRecorder");
+	TakeRecorderModule.GetLastLevelSequenceProvider().BindUObject(this, &UTakeRecorderSubsystemImplementation::SetLastLevelSequence);
+
+	bAutoApplyTakeNumber = true;
+	
+	CacheMetaData();
+	
+	if (TakeMetaData && !TakeMetaData->IsLocked())
+	{
+		const int32 NextTakeNumber = GetNextTakeNumber(TakeMetaData->GetSlate());
+		if (NextTakeNumber != TakeMetaData->GetTakeNumber())
+		{
+			TakeMetaData->SetTakeNumber(NextTakeNumber);
+		}
+	}
+	
+	// If a recording is currently underway, initialize to that now
+	if (UTakeRecorder* ActiveRecorder = UTakeRecorder::GetActiveRecorder())
+	{
+		RecordingLevelSequence = ActiveRecorder->GetSequence();
+
+		ActiveRecorder->OnRecordingStopped().Remove(OnRecordingStoppedHandle);
+		ActiveRecorder->OnRecordingFinished().Remove(OnRecordingFinishedHandle);
+		ActiveRecorder->OnRecordingCancelled().Remove(OnRecordingCancelledHandle);
+		OnRecordingStoppedHandle = ActiveRecorder->OnRecordingStopped().AddUObject(this, &UTakeRecorderSubsystemImplementation::OnRecordingStopped);
+		OnRecordingFinishedHandle = ActiveRecorder->OnRecordingFinished().AddUObject(this, &UTakeRecorderSubsystemImplementation::OnRecordingFinished);
+		OnRecordingCancelledHandle = ActiveRecorder->OnRecordingCancelled().AddUObject(this, &UTakeRecorderSubsystemImplementation::OnRecordingCancelled);
+	}
+	else
+	{
+		RecordingLevelSequence = nullptr;
+	}
+	
+	TransientPreset = AllocateTransientPreset();
+	
+	// Copy the base preset into the transient preset if it was provided.
+	// We do this first so that anything that asks for its Level Sequence
+	// on construction gets the right one
+	if (InData.BasePreset)
+	{
+		TransientPreset->CopyFrom(InData.BasePreset);
+	}
+	else if (InData.BaseSequence)
+	{
+		TransientPreset->CopyFrom(InData.BaseSequence);
+
+		ULevelSequence* LevelSequence = TransientPreset->GetLevelSequence();
+
+		UTakeRecorderSources* BaseSources = InData.BaseSequence->FindMetaData<UTakeRecorderSources>();
+		if (BaseSources && LevelSequence)
+		{
+			LevelSequence->CopyMetaData(BaseSources);
+		}
+
+		if (LevelSequence)
+		{
+			LevelSequence->GetMovieScene()->SetReadOnly(false);
+		}
+
+		if (UTakeMetaData*  TakeMetaDataLevelSequence = LevelSequence ? LevelSequence->FindMetaData<UTakeMetaData>() : nullptr)
+		{
+			TakeMetaDataLevelSequence->Unlock();
+			TakeMetaDataLevelSequence->SetTimestamp(FDateTime(0));
+		}
+	}
+	else if (InData.RecordIntoSequence)
+	{
+		SetRecordIntoLevelSequence(InData.RecordIntoSequence);
+	}
+	else if (InData.SequenceToView)
+	{
+		SuppliedLevelSequence = InData.SequenceToView;
+	}
+
+	/** Clear the dirty flag since the preset was just initialized. */
+	TransientPreset->GetOutermost()->SetDirtyFlag(false);
+
+	bHasTargetSequenceBeenSet = true;
+	SetTickableTickType(ETickableTickType::Conditional);
+}
+
+void UTakeRecorderSubsystemImplementation::SetRecordIntoLevelSequence(ULevelSequence* LevelSequence)
+{
+	SuppliedLevelSequence = nullptr;
+	RecordIntoLevelSequence = LevelSequence;
+
+	if (RecordIntoLevelSequence)
+	{
+		RecordIntoLevelSequence->GetMovieScene()->SetReadOnly(false);
+	}
+
+	if (UTakeMetaData* LocalTakeMetaData = RecordIntoLevelSequence ? RecordIntoLevelSequence->FindOrAddMetaData<UTakeMetaData>() : nullptr)
+	{
+		// Set up take metadata to match this level sequence's info, ie. match the frame rate, use the level sequence name as the slate
+		LocalTakeMetaData->Unlock();
+		LocalTakeMetaData->SetTimestamp(FDateTime(0));
+		LocalTakeMetaData->SetSlate(LevelSequence->GetName());
+		LocalTakeMetaData->SetTakeNumber(0);
+		LocalTakeMetaData->SetFrameRate(LevelSequence->GetMovieScene()->GetDisplayRate());
+		LocalTakeMetaData->SetFrameRateFromTimecode(false);
+	}
+}
+
+bool UTakeRecorderSubsystemImplementation::CanReviewLastRecording() const
+{
+	if (GetLastRecordedLevelSequence() == nullptr)
+	{
+		return false;
+	}
+	ITakeRecorderModule& TakeRecorderModule = FModuleManager::Get().LoadModuleChecked<ITakeRecorderModule>("TakeRecorder");
+	FCanReviewLastRecordedLevelSequence& CanReview = TakeRecorderModule.GetCanReviewLastRecordedLevelSequenceDelegate();
+	if (CanReview.IsBound())
+	{
+		return CanReview.Execute();
+	}
+	
+	return true;
+}
+
+bool UTakeRecorderSubsystemImplementation::ReviewLastRecording()
+{
+	if (LastRecordedLevelSequence.IsValid())
+	{
+		SuppliedLevelSequence = LastRecordedLevelSequence.Get();
+		LastRecordedLevelSequence = nullptr;
+		return true;
+	}
+
+	return false;
+}
+
+bool UTakeRecorderSubsystemImplementation::StartRecording(bool bOpenSequencer, bool bShowErrorMessage)
+{
+	static bool bStartedRecording = false;
+
+	if (bStartedRecording)
+	{
+		return false;
+	}
+
+	if (!GIsEditor)
+	{
+		// Sequencer requires an editor attached.
+		bOpenSequencer = false;
+	}
+
+	TGuardValue<bool> ReentrantGuard(bStartedRecording, true);
+
+	ULevelSequence*       LevelSequence = GetLevelSequence();
+	UTakeRecorderSources* Sources = LevelSequence ? LevelSequence->FindMetaData<UTakeRecorderSources>() : nullptr;
+	
+	if (LevelSequence && Sources)
+	{
+		FTakeRecorderParameters Parameters;
+		Parameters.User    = GetDefault<UTakeRecorderUserSettings>()->Settings;
+		Parameters.Project = GetDefault<UTakeRecorderProjectSettings>()->Settings;
+		Parameters.TakeRecorderMode = GetTakeRecorderMode();
+		Parameters.StartFrame = LevelSequence->GetMovieScene()->GetPlaybackRange().GetLowerBoundValue();
+		Parameters.bOpenSequencer = bOpenSequencer;
+
+		FText ErrorText = LOCTEXT("UnknownError", "An unknown error occurred when trying to start recording");
+
+		if (GIsEditor)
+		{
+			IAssetEditorInstance* AssetEditor = GEditor->GetEditorSubsystem<UAssetEditorSubsystem>()->FindEditorForAsset(LevelSequence, false);
+			ILevelSequenceEditorToolkit* LevelSequenceEditor = static_cast<ILevelSequenceEditorToolkit*>(AssetEditor);
+
+			if (LevelSequenceEditor && LevelSequenceEditor->GetSequencer())
+			{
+				// If not resetting the playhead, store the current time as the start frame for recording. 
+				// This will ultimately be the start of the playback range and the recording will begin from that time.
+				if (!Parameters.User.bResetPlayhead)
+				{
+					Parameters.StartFrame = LevelSequenceEditor->GetSequencer()->GetLocalTime().Time.FrameNumber;
+				}
+			}
+		}
+
+		UTakeRecorder* NewRecorder = NewObject<UTakeRecorder>(GetTransientPackage(), NAME_None, RF_Transient);
+
+		if (!NewRecorder->Initialize(LevelSequence, Sources, TakeMetaData, Parameters, &ErrorText))
+		{
+			if (bShowErrorMessage && ensure(!ErrorText.IsEmpty()))
+			{
+				FNotificationInfo Info(ErrorText);
+				Info.ExpireDuration = 5.0f;
+				FSlateNotificationManager::Get().AddNotification(Info)->SetCompletionState(SNotificationItem::CS_Fail);
+			}
+			return false;
+		}
+
+		OnRecordingStarted(NewRecorder);
+		
+		return true;
+	}
+	
+	return false;
+}
+
+void UTakeRecorderSubsystemImplementation::StopRecording()
+{
+	if (UTakeRecorder* CurrentRecording = UTakeRecorder::GetActiveRecorder())
+	{
+		CurrentRecording->Stop();
+	}
+}
+
+void UTakeRecorderSubsystemImplementation::CancelRecording(bool bShowConfirmMessage)
+{
+	if (UTakeRecorder* CurrentRecording = UTakeRecorder::GetActiveRecorder())
+	{
+		auto DoCancel = [](const TWeakObjectPtr<UTakeRecorder>& WeakRecording)
+		{
+			if (WeakRecording.IsValid())
+			{
+				WeakRecording.Get()->Cancel(); 
+			}
+		};
+
+		if (bShowConfirmMessage)
+		{
+			const TSharedRef<SMessageDialog> ConfirmDialog = SNew(SMessageDialog)
+				.Title(FText(LOCTEXT("ConfirmCancelRecordingTitle", "Cancel Recording?")))
+				.Message(LOCTEXT("ConfirmCancelRecording", "Are you sure you want to cancel the current recording?"))
+				.Buttons({
+					SCustomDialog::FButton(LOCTEXT("Yes", "Yes"), FSimpleDelegate::CreateLambda([DoCancel, WeakRecording = TWeakObjectPtr<UTakeRecorder>(CurrentRecording)]()
+						{ 
+							DoCancel(WeakRecording);
+						})),
+					SCustomDialog::FButton(LOCTEXT("No", "No"))
+				});
+
+			// Non modal so that the recording continues to update
+			ConfirmDialog->Show();
+		}
+		else
+		{
+			DoCancel(TWeakObjectPtr<UTakeRecorder>(CurrentRecording));
+		}
+	}
+}
+
+void UTakeRecorderSubsystemImplementation::ResetToPendingTake()
+{
+	if (IsReviewing())
+	{
+		LastRecordedLevelSequence = SuppliedLevelSequence;
+	}
+
+	SuppliedLevelSequence = nullptr;
+	RecordIntoLevelSequence = nullptr;
+	
+	TransientPreset = AllocateTransientPreset();
+}
+
+void UTakeRecorderSubsystemImplementation::ClearPendingTake()
+{
+	if (IsReviewing())
+	{
+		LastRecordedLevelSequence = SuppliedLevelSequence;
+	}
+
+	UTakeRecorderSources* BaseSources = nullptr;
+
+	if (const ULevelSequence* CurrentLevelSequence = GetLevelSequence())
+	{
+		BaseSources = CurrentLevelSequence->FindMetaData<UTakeRecorderSources>();
+	}
+
+	SuppliedLevelSequence = nullptr;
+	RecordIntoLevelSequence = nullptr;
+
+	FScopedTransaction Transaction(LOCTEXT("ClearPendingTake_Transaction", "Clear Pending Take"));
+
+	TransientPreset->Modify();
+	TransientPreset->CreateLevelSequence();
+
+	ULevelSequence* LevelSequence = TransientPreset->GetLevelSequence();
+	if (LevelSequence && BaseSources)
+	{
+		LevelSequence->CopyMetaData(BaseSources);
+	}
+}
+
+UTakePreset* UTakeRecorderSubsystemImplementation::GetPendingTake() const
+{
+	const ITakeRecorderModule& TakeRecorderModule = FModuleManager::LoadModuleChecked<ITakeRecorderModule>("TakeRecorder");
+	return TakeRecorderModule.GetPendingTake();
+}
+
+void UTakeRecorderSubsystemImplementation::RevertChanges()
+{
+	UTakePreset* PresetOrigin = GetTakeMetaData()->GetPresetOrigin();
+
+	FScopedTransaction Transaction(LOCTEXT("RevertChanges_Transaction", "Revert Changes"));
+
+	TransientPreset->Modify();
+	TransientPreset->CopyFrom(PresetOrigin);
+	TransientPreset->GetOutermost()->SetDirtyFlag(false);
+}
+
+UTakeRecorderSource* UTakeRecorderSubsystemImplementation::AddSource(const TSubclassOf<UTakeRecorderSource> InSourceClass)
+{
+	UTakeRecorderSources* Sources = GetSources();
+
+	if (*InSourceClass && Sources)
+	{
+		FScopedTransaction Transaction(FText::Format(LOCTEXT("AddNewSource", "Add New {0} Source"), InSourceClass->GetDisplayNameText()));
+		Sources->Modify();
+
+		return Sources->AddSource(InSourceClass);
+	}
+
+	return nullptr;
+}
+
+void UTakeRecorderSubsystemImplementation::RemoveSource(UTakeRecorderSource* InSource)
+{
+	UTakeRecorderSources* Sources = GetSources();
+	if (Sources && InSource)
+	{
+		Sources->RemoveSource(InSource);
+	}
+}
+
+void UTakeRecorderSubsystemImplementation::ClearSources()
+{
+	if (UTakeRecorderSources* Sources = GetSources())
+	{
+		for (UTakeRecorderSource* Source : Sources->GetSourcesCopy())
+		{
+			Sources->RemoveSource(Source);
+		}
+	}
+}
+
+UTakeRecorderSources* UTakeRecorderSubsystemImplementation::GetSources() const
+{
+	ULevelSequence* LevelSequence = GetLevelSequence();
+	UTakeRecorderSources* Sources = LevelSequence ? LevelSequence->FindOrAddMetaData<UTakeRecorderSources>() : nullptr;
+	return Sources;
+}
+
+TArrayView<UTakeRecorderSource* const> UTakeRecorderSubsystemImplementation::GetAllSources() const
+{
+	if (const UTakeRecorderSources* Sources = GetSources())
+	{
+		return Sources->GetSources();
+	}
+
+	return {};
+}
+
+TArray<UTakeRecorderSource*> UTakeRecorderSubsystemImplementation::GetAllSourcesCopy() const
+{
+	return TArray<UTakeRecorderSource*>(GetAllSources());
+}
+
+UTakeRecorderSource* UTakeRecorderSubsystemImplementation::GetSourceByClass(const TSubclassOf<UTakeRecorderSource> InSourceClass) const
+{
+	const TArrayView<UTakeRecorderSource* const> Sources = GetAllSources();
+	for (UTakeRecorderSource* Source : Sources)
+	{
+		if (Source->IsA(InSourceClass))
+		{
+			return Source;
+		}
+	}
+
+	return nullptr;
+}
+
+void UTakeRecorderSubsystemImplementation::AddSourceForActor(AActor* InActor, bool bReduceKeys, bool bShowProgress)
+{
+	if (UTakeRecorderSources* Sources = GetSources())
+	{
+		FTakeRecorderModule& TakeRecorderModule = FModuleManager::GetModuleChecked<FTakeRecorderModule>("TakeRecorder");
+		TakeRecorderModule.GetSourcesExtensionData().OnAddActorSource.ExecuteIfBound(Sources, { InActor }, bReduceKeys, bShowProgress);
+	}
+}
+
+void UTakeRecorderSubsystemImplementation::RemoveActorFromSources(AActor* InActor)
+{
+	if (UTakeRecorderSources* Sources = GetSources())
+	{
+		FTakeRecorderModule& TakeRecorderModule = FModuleManager::GetModuleChecked<FTakeRecorderModule>("TakeRecorder");
+		TakeRecorderModule.GetSourcesExtensionData().OnRemoveActorSource.ExecuteIfBound(Sources, { InActor });
+	}
+}
+
+AActor* UTakeRecorderSubsystemImplementation::GetSourceActor(UTakeRecorderSource* InSource) const
+{
+	FTakeRecorderModule& TakeRecorderModule = FModuleManager::GetModuleChecked<FTakeRecorderModule>("TakeRecorder");
+	if (TakeRecorderModule.GetSourcesExtensionData().OnGetSourceActor.IsBound())
+	{
+		return TakeRecorderModule.GetSourcesExtensionData().OnGetSourceActor.Execute(InSource);
+	}
+
+	return nullptr;
+}
+
+ETakeRecorderState UTakeRecorderSubsystemImplementation::GetState() const
+{
+	if (const UTakeRecorder* Recorder = UTakeRecorder::GetActiveRecorder())
+	{
+		return Recorder->GetState();
+	}
+
+	return ETakeRecorderState::PreInitialization;
+}
+
+void UTakeRecorderSubsystemImplementation::SetTakeNumber(int32 InNewTakeNumber, bool bEmitChanged)
+{
+	if (TakeMetaData)
+	{
+		FScopedTransaction Transaction(LOCTEXT("SetTake_Transaction", "Set Take Number"));
+		TakeMetaData->Modify();
+		TakeMetaData->SetTakeNumber(InNewTakeNumber, bEmitChanged);
+		bAutoApplyTakeNumber = false;
+	}
+}
+
+int32 UTakeRecorderSubsystemImplementation::GetNextTakeNumber(const FString& InSlate) const
+{
+	return UTakesCoreBlueprintLibrary::ComputeNextTakeNumber(InSlate);
+}
+
+void UTakeRecorderSubsystemImplementation::GetNumberOfTakes(const FString& InSlate, int32& OutMaxTake, int32& OutNumTakes) const
+{
+	int32 MaxTake = 0;
+
+	TArray<FAssetData> Takes = UTakesCoreBlueprintLibrary::FindTakes(InSlate);
+
+	for (const FAssetData& Asset : Takes)
+	{
+		FAssetDataTagMapSharedView::FFindTagResult TakeNumberTag = Asset.TagsAndValues.FindTag(UTakeMetaData::AssetRegistryTag_TakeNumber);
+
+		int32 ThisTakeNumber = 0;
+		if (TakeNumberTag.IsSet() && LexTryParseString(ThisTakeNumber, *TakeNumberTag.GetValue()))
+		{
+			MaxTake = FMath::Max(MaxTake, ThisTakeNumber);
+		}
+	}
+
+	OutMaxTake = MaxTake;
+	OutNumTakes = Takes.Num();
+}
+
+TArray<FAssetData> UTakeRecorderSubsystemImplementation::GetSlates(FName InPackagePath) const
+{
+	const IAssetRegistry& AssetRegistry = FModuleManager::LoadModuleChecked<FAssetRegistryModule>(AssetRegistryConstants::ModuleName).Get();
+
+	FARFilter Filter;
+	Filter.ClassPaths.Add(ULevelSequence::StaticClass()->GetClassPathName());
+	Filter.TagsAndValues.Add(UTakeMetaData::AssetRegistryTag_Slate);
+	if (!InPackagePath.IsNone())
+	{
+		Filter.PackagePaths.Add(InPackagePath);
+		Filter.bRecursivePaths = true;
+	}
+
+	TArray<FAssetData> AllAssets;
+	AssetRegistry.GetAssets(Filter, AllAssets);
+	
+	return AllAssets;
+}
+
+void UTakeRecorderSubsystemImplementation::SetSlateName(const FString& InSlateName, bool bEmitChanged)
+{
+	if (TakeMetaData && TakeMetaData->GetSlate() != InSlateName)
+	{
+		FScopedTransaction Transaction(LOCTEXT("SetSlate_Transaction", "Set Take Slate"));
+		TakeMetaData->Modify();
+
+		TakeMetaData->SetSlate(InSlateName, bEmitChanged);
+
+		// Compute the correct starting take number
+		const int32 NextTakeNumber = GetNextTakeNumber(TakeMetaData->GetSlate());
+		if (NextTakeNumber != TakeMetaData->GetTakeNumber())
+		{
+			TakeMetaData->SetTakeNumber(NextTakeNumber, bEmitChanged);
+		}
+	}
+}
+
+bool UTakeRecorderSubsystemImplementation::MarkFrame()
+{
+	if (UTakeRecorderBlueprintLibrary::IsRecording() && TakeMetaData)
+	{
+		const FFrameRate FrameRate = TakeMetaData->GetFrameRate();
+
+		const FTimespan RecordingDuration = FDateTime::UtcNow() - TakeMetaData->GetTimestamp();
+
+		const FFrameNumber ElapsedFrame = FFrameNumber(static_cast<int32>(FrameRate.AsDecimal() * RecordingDuration.GetTotalSeconds()));
+
+		const ULevelSequence* LevelSequence = GetLevelSequence();
+		if (!LevelSequence)
+		{
+			return false;
+		}
+
+		UMovieScene* MovieScene = LevelSequence->GetMovieScene();
+		if (!MovieScene)
+		{
+			return false;
+		}
+
+		const FFrameRate DisplayRate = MovieScene->GetDisplayRate();
+		const FFrameRate TickResolution = MovieScene->GetTickResolution();
+
+		FMovieSceneMarkedFrame MarkedFrame;
+
+		const UTakeRecorderSources* Sources = LevelSequence->FindMetaData<UTakeRecorderSources>();
+		if (Sources && Sources->GetSettings().bStartAtCurrentTimecode)
+		{
+			MarkedFrame.FrameNumber = FFrameRate::TransformTime(FFrameTime(FApp::GetTimecode().ToFrameNumber(DisplayRate)), DisplayRate, TickResolution).FloorToFrame();
+		}
+		else
+		{
+			MarkedFrame.FrameNumber = ConvertFrameTime(ElapsedFrame, DisplayRate, TickResolution).CeilToFrame();
+		}
+
+		const int32 MarkedFrameIndex = MovieScene->AddMarkedFrame(MarkedFrame);
+		UTakeRecorderBlueprintLibrary::OnTakeRecorderMarkedFrameAdded(MovieScene->GetMarkedFrames()[MarkedFrameIndex]);
+		if (OwningSubsystemWeakPtr.IsValid())
+		{
+			OwningSubsystemWeakPtr->TakeRecorderMarkedFrameAdded.Broadcast(MovieScene->GetMarkedFrames()[MarkedFrameIndex]);
+		}
+
+		return true;
+	}
+
+	return false;
+}
+
+FFrameRate UTakeRecorderSubsystemImplementation::GetFrameRate() const
+{
+	if (TakeMetaData)
+	{
+		return TakeMetaData->GetFrameRate();
+	}
+
+	return FFrameRate();
+}
+
+void UTakeRecorderSubsystemImplementation::SetFrameRate(FFrameRate InFrameRate)
+{
+	SetFrameRateImpl(InFrameRate, false);
+}
+
+void UTakeRecorderSubsystemImplementation::SetFrameRateFromTimecode()
+{
+	SetFrameRateImpl(FApp::GetTimecodeFrameRate(), true);
+}
+
+void UTakeRecorderSubsystemImplementation::ImportPreset(const FAssetData& InPreset)
+{
+	SuppliedLevelSequence = nullptr;
+	RecordIntoLevelSequence = nullptr;
+
+	UTakePreset* Take = CastChecked<UTakePreset>(InPreset.GetAsset());
+	FScopedTransaction Transaction(LOCTEXT("ImportPreset_Transaction", "Import Take Preset"));
+
+	TransientPreset->Modify();
+	TransientPreset->CopyFrom(Take);
+	TransientPreset->GetOutermost()->SetDirtyFlag(false);
+
+	GetTakeMetaData()->SetPresetOrigin(Take);
+}
+
+bool UTakeRecorderSubsystemImplementation::IsReviewing() const
+{
+	return !IsRecording() && (TakeMetaData && TakeMetaData->Recorded() && GetTakeRecorderMode() != ETakeRecorderMode::RecordIntoSequence);
+}
+
+bool UTakeRecorderSubsystemImplementation::IsRecording() const
+{
+	return UTakeRecorderBlueprintLibrary::GetActiveRecorder() ? true : false;
+}
+
+bool UTakeRecorderSubsystemImplementation::TryGetSequenceCountdown(float& OutValue) const
+{
+	const UTakeRecorder* CurrentRecording = UTakeRecorder::GetActiveRecorder();
+	const bool           bIsCountingDown  = CurrentRecording && CurrentRecording->GetState() == ETakeRecorderState::CountingDown;
+	
+	OutValue = bIsCountingDown ? CurrentRecording->GetCountdownSeconds() : 0.f;
+	return bIsCountingDown;
+}
+
+void UTakeRecorderSubsystemImplementation::SetSequenceCountdown(float InSeconds)
+{
+	if (UTakeRecorder* CurrentRecording = UTakeRecorder::GetActiveRecorder())
+	{
+		CurrentRecording->SetCountdown(InSeconds);
+	}
+}
+
+TArray<UObject*> UTakeRecorderSubsystemImplementation::GetSourceRecordSettings(UTakeRecorderSource* InSource) const
+{
+	return InSource->GetAdditionalSettingsObjects();
+}
+
+FTakeRecorderParameters UTakeRecorderSubsystemImplementation::GetGlobalRecordSettings() const
+{
+	return UTakeRecorderBlueprintLibrary::GetDefaultParameters();
+}
+
+void UTakeRecorderSubsystemImplementation::SetGlobalRecordSettings(const FTakeRecorderParameters& InParameters)
+{
+	UTakeRecorderBlueprintLibrary::SetDefaultParameters(InParameters);
+}
+
+UTakeMetaData* UTakeRecorderSubsystemImplementation::GetTakeMetaData() const
+{
+	return TakeMetaData;
+}
+
+ULevelSequence* UTakeRecorderSubsystemImplementation::GetLevelSequence() const
+{
+	if (SuppliedLevelSequence)
+	{
+		return SuppliedLevelSequence;
+	}
+	else if (RecordIntoLevelSequence)
+	{
+		return RecordIntoLevelSequence;
+	}
+	else if (RecordingLevelSequence)
+	{
+		return RecordingLevelSequence;
+	}
+	else if (TransientPreset)
+	{
+		return TransientPreset->GetLevelSequence();
+	}
+	return nullptr;
+}
+
+ULevelSequence* UTakeRecorderSubsystemImplementation::GetSuppliedLevelSequence() const
+{
+	return SuppliedLevelSequence;
+}
+
+ULevelSequence* UTakeRecorderSubsystemImplementation::GetRecordingLevelSequence() const
+{
+	return RecordingLevelSequence;
+}
+
+ULevelSequence* UTakeRecorderSubsystemImplementation::GetRecordIntoLevelSequence() const
+{
+	return RecordIntoLevelSequence;
+}
+
+ULevelSequence* UTakeRecorderSubsystemImplementation::GetLastRecordedLevelSequence() const
+{
+	return LastRecordedLevelSequence.Get();
+}
+
+UTakePreset* UTakeRecorderSubsystemImplementation::GetTransientPreset() const
+{
+	return TransientPreset;
+}
+
+ETakeRecorderMode UTakeRecorderSubsystemImplementation::GetTakeRecorderMode() const
+{
+	if (RecordIntoLevelSequence != nullptr)
+	{
+		return ETakeRecorderMode::RecordIntoSequence;
+	}
+
+	return ETakeRecorderMode::RecordNewSequence;
+}
+
+UTakeRecorderNamingTokensData* UTakeRecorderSubsystemImplementation::GetNamingTokensData() const
+{
+	check(NamingTokensData);
+	return NamingTokensData;
+}
+
+bool UTakeRecorderSubsystemImplementation::HasPendingChanges() const
+{
+	UMovieScene* MovieScene = GetLevelSequence() ? GetLevelSequence()->GetMovieScene() : nullptr;
+	const bool bHasUserMadeChanges = MovieScene
+		&& (!MovieScene->GetTracks().IsEmpty()
+			|| !MovieScene->GetBindings().IsEmpty()
+			|| MovieScene->GetPossessableCount() > 0
+			|| MovieScene->GetSpawnableCount() > 0);
+	return bHasUserMadeChanges;
+}
+
+UTakePreset* UTakeRecorderSubsystemImplementation::AllocateTransientPreset()
+{
+	return UTakePreset::AllocateTransientPreset(GetDefault<UTakeRecorderUserSettings>()->LastOpenedPreset.Get());
+}
+
+void UTakeRecorderSubsystemImplementation::CacheMetaData()
+{
+	UTakeMetaData* NewMetaDataThisTick = nullptr;
+
+	if (const ULevelSequence* LevelSequence = GetLevelSequence())
+	{
+		NewMetaDataThisTick = LevelSequence->FindMetaData<UTakeMetaData>();
+	}
+
+	// If it's null we use the transient meta-data
+	if (!NewMetaDataThisTick)
+	{
+		// if the transient meta-data doesn't exist, create it now
+		if (!TransientTakeMetaData)
+		{
+			TransientTakeMetaData = UTakeMetaData::CreateFromDefaults(GetTransientPackage(), NAME_None);
+			TransientTakeMetaData->SetFlags(RF_Transactional | RF_Transient);
+
+			FString DefaultSlate = GetDefault<UTakeRecorderProjectSettings>()->Settings.DefaultSlate;
+			if (TransientTakeMetaData->GetSlate() != DefaultSlate)
+			{
+				TransientTakeMetaData->SetSlate(DefaultSlate, false);
+			}
+
+			// Compute the correct starting take number
+			int32 NextTakeNumber = GetNextTakeNumber(TransientTakeMetaData->GetSlate());
+			if (TransientTakeMetaData->GetTakeNumber() != NextTakeNumber)
+			{
+				TransientTakeMetaData->SetTakeNumber(NextTakeNumber, false);
+			}
+
+			UTakeMetaData::SetMostRecentMetaData(TransientTakeMetaData);
+		}
+
+		NewMetaDataThisTick = TransientTakeMetaData;
+	}
+
+	check(NewMetaDataThisTick);
+	if (NewMetaDataThisTick != TakeMetaData)
+	{
+		TakeMetaData = NewMetaDataThisTick;
+
+		if (!OnPreForceDeleteObjectsHandle.IsValid())
+		{
+			// Since this is a strong reference, we need to be able to clear it if the user is attempting to remove it, but let's not always
+			// hook this event since the subsystem won't be active for general UE operation.
+			OnPreForceDeleteObjectsHandle = FEditorDelegates::OnPreForceDeleteObjects.AddUObject(this, &UTakeRecorderSubsystemImplementation::OnPreForceDeleteObjects);
+		}
+	}
+
+	//Set MovieScene Display Rate to the Preset Frame Rate.
+	if (const ULevelSequence* Sequence = GetLevelSequence())
+	{
+		if (UMovieScene* MovieScene = Sequence->GetMovieScene())
+		{
+			MovieScene->SetDisplayRate(TakeMetaData->GetFrameRate());
+		}
+	}
+
+	check(TakeMetaData);
+}
+
+void UTakeRecorderSubsystemImplementation::SetFrameRateImpl(const FFrameRate& InFrameRate, bool bFromTimecode)
+{
+	if (TakeMetaData)
+	{
+		TakeMetaData->SetFrameRateFromTimecode(bFromTimecode);
+		TakeMetaData->SetFrameRate(InFrameRate);
+	}
+	const ULevelSequence* Sequence = GetLevelSequence();
+	if (UMovieScene* MovieScene = Sequence ? Sequence->GetMovieScene() : nullptr)
+	{
+		MovieScene->SetDisplayRate(InFrameRate);
+	}
+}
+
+void UTakeRecorderSubsystemImplementation::IncrementTakeNumber()
+{
+	if (TransientTakeMetaData)
+	{
+		// Increment the transient take meta data if necessary
+		int32 NextTakeNumber = GetNextTakeNumber(TransientTakeMetaData->GetSlate());
+
+		if (TransientTakeMetaData->GetTakeNumber() != NextTakeNumber)
+		{
+			TransientTakeMetaData->SetTakeNumber(NextTakeNumber);
+		}
+	}
+	
+	// Update the preset take number at the end of recording
+	
+	if (const ULevelSequence* LevelSequence = TransientPreset ? TransientPreset->GetLevelSequence() : nullptr)
+	{
+		if (UTakeMetaData* MetaData = LevelSequence->FindMetaData<UTakeMetaData>())
+		{
+			const int32 NextTakeNumber = GetNextTakeNumber(MetaData->GetSlate());
+			MetaData->SetTakeNumber(NextTakeNumber);
+		}
+	}
+
+	bAutoApplyTakeNumber = true;
+}
+
+void UTakeRecorderSubsystemImplementation::OnAssetRegistryFilesLoaded()
+{
+	if (bAutoApplyTakeNumber && TransientTakeMetaData)
+	{
+		const int32 NextTakeNumber = GetNextTakeNumber(TransientTakeMetaData->GetSlate());
+		TransientTakeMetaData->SetTakeNumber(NextTakeNumber);
+	}
+}
+
+void UTakeRecorderSubsystemImplementation::OnRecordingInitialized(UTakeRecorder* Recorder)
+{
+	if (!bHasTargetSequenceBeenSet)
+	{
+		// Not initialized, take recorder triggered outside of the subsystem.
+		SetTargetSequence();
+	}
+	
+	// This needs to be stored with a strong ptr before our panels refresh, otherwise the weak sequencer ref will be lost during initialize.
+	RecordingLevelSequence = Recorder->GetSequence();
+	// Recache the meta-data here since we know that the sequence has probably changed as a result of the recording being started
+	CacheMetaData();
+	OnRecordingFinishedHandle = Recorder->OnRecordingFinished().AddUObject(this, &UTakeRecorderSubsystemImplementation::OnRecordingFinished);
+	OnRecordingCancelledHandle = Recorder->OnRecordingCancelled().AddUObject(this, &UTakeRecorderSubsystemImplementation::OnRecordingCancelled);
+
+	WITH_OWNING_SUBSYSTEM(GetOnRecordingInitializedEvent().Broadcast(Recorder));
+	WITH_OWNING_SUBSYSTEM(TakeRecorderInitialized.Broadcast());
+}
+
+
+void UTakeRecorderSubsystemImplementation::OnRecordingStarted(UTakeRecorder* Recorder)
+{
+	WITH_OWNING_SUBSYSTEM(GetOnRecordingStartedEvent().Broadcast(Recorder));
+	WITH_OWNING_SUBSYSTEM(TakeRecorderStarted.Broadcast());
+}
+
+void UTakeRecorderSubsystemImplementation::OnRecordingStopped(UTakeRecorder* Recorder)
+{
+	Recorder->OnRecordingStopped().Remove(OnRecordingStoppedHandle);
+
+	WITH_OWNING_SUBSYSTEM(GetOnRecordingStoppedEvent().Broadcast(Recorder));
+	WITH_OWNING_SUBSYSTEM(TakeRecorderStopped.Broadcast());
+}
+
+void UTakeRecorderSubsystemImplementation::OnRecordingFinished(UTakeRecorder* Recorder)
+{
+	LastRecordedLevelSequence = RecordingLevelSequence;
+	RecordingLevelSequence = nullptr;
+
+	IncrementTakeNumber();
+	
+	if (TransientTakeMetaData)
+	{
+		// Increment the transient take meta data if necessary
+		int32 NextTakeNumber = GetNextTakeNumber(TransientTakeMetaData->GetSlate());
+
+		if (TransientTakeMetaData->GetTakeNumber() != NextTakeNumber)
+		{
+			TransientTakeMetaData->SetTakeNumber(NextTakeNumber);
+		}
+
+		bAutoApplyTakeNumber = true;
+	}
+
+	Recorder->OnRecordingFinished().Remove(OnRecordingFinishedHandle);
+
+	WITH_OWNING_SUBSYSTEM(GetOnRecordingFinishedEvent().Broadcast(Recorder));
+	WITH_OWNING_SUBSYSTEM(TakeRecorderFinished.Broadcast(LastRecordedLevelSequence.Get()));
+}
+
+void UTakeRecorderSubsystemImplementation::OnRecordingCancelled(UTakeRecorder* Recorder)
+{
+	RecordingLevelSequence = nullptr;
+
+	Recorder->OnRecordingFinished().Remove(OnRecordingFinishedHandle);
+	Recorder->OnRecordingCancelled().Remove(OnRecordingCancelledHandle);
+
+	WITH_OWNING_SUBSYSTEM(GetOnRecordingCancelledEvent().Broadcast(Recorder));
+	WITH_OWNING_SUBSYSTEM(TakeRecorderCancelled.Broadcast());
+}
+
+void UTakeRecorderSubsystemImplementation::OnTakeSlateChanged(const FString& InSlate, UTakeMetaData* InTakeMetaData)
+{
+	WITH_OWNING_SUBSYSTEM(TakeRecorderSlateChanged.Broadcast(InSlate, InTakeMetaData));
+}
+
+void UTakeRecorderSubsystemImplementation::OnTakeNumberChanged(int32 InTakeNumber, UTakeMetaData* InTakeMetaData)
+{
+	WITH_OWNING_SUBSYSTEM(TakeRecorderTakeNumberChanged.Broadcast(InTakeNumber, InTakeMetaData));
+}
+
+void UTakeRecorderSubsystemImplementation::SetLastLevelSequence(ULevelSequence* InSequence)
+{
+	LastRecordedLevelSequence = InSequence;
+}
+
+void UTakeRecorderSubsystemImplementation::OnSourceAdded(UTakeRecorderSource* Source)
+{
+	WITH_OWNING_SUBSYSTEM(GetOnRecordingSourceAddedEvent().Broadcast(Source));
+	WITH_OWNING_SUBSYSTEM(TakeRecorderSourceAdded.Broadcast(Source));
+}
+
+void UTakeRecorderSubsystemImplementation::OnSourceRemoved(UTakeRecorderSource* Source)
+{
+	WITH_OWNING_SUBSYSTEM(GetOnRecordingSourceRemovedEvent().Broadcast(Source));
+	WITH_OWNING_SUBSYSTEM(TakeRecorderSourceRemoved.Broadcast(Source));
+}
+
+void UTakeRecorderSubsystemImplementation::BindToNamingTokenEvents()
+{
+	if (GEngine && !GetNamingTokensData()->TakeRecorderNamingTokens.IsValid())
+	{
+		GetNamingTokensData()->TakeRecorderNamingTokens = GEngine->GetEngineSubsystem<UNamingTokensEngineSubsystem>()->GetNamingTokens(ITakeRecorderNamingTokensModule::GetTakeRecorderNamespace());
+		if (ensure(GetNamingTokensData()->TakeRecorderNamingTokens.IsValid()))
+		{
+			const TStrongObjectPtr<UNamingTokens> TakeRecorderNamingTokensPin = GetNamingTokensData()->TakeRecorderNamingTokens.Pin();
+			TakeRecorderNamingTokensPin->RegisterExternalTokens(GetNamingTokensData()->NamingTokensExternalGuid);
+			TakeRecorderNamingTokensPin->GetOnPreEvaluateEvent().AddUObject(this, &UTakeRecorderSubsystemImplementation::OnTakeRecorderNamingTokensPreEvaluate);
+		}
+	}
+}
+
+void UTakeRecorderSubsystemImplementation::UnbindNamingTokensEvents()
+{
+	if (const TStrongObjectPtr<UNamingTokens> TakeRecorderNamingTokensPin = GetNamingTokensData()->TakeRecorderNamingTokens.Pin())
+	{
+		TakeRecorderNamingTokensPin->UnregisterExternalTokens(GetNamingTokensData()->NamingTokensExternalGuid);
+		TakeRecorderNamingTokensPin->GetOnPreEvaluateEvent().RemoveAll(this);
+	}
+}
+
+void UTakeRecorderSubsystemImplementation::OnTakeRecorderNamingTokensPreEvaluate(const FNamingTokensEvaluationData& InEvaluationData)
+{
+	if (GetNamingTokensData()->TakeRecorderNamingTokens.IsValid())
+	{
+		TArray<FNamingTokenData>& ExternalTokens = GetNamingTokensData()->TakeRecorderNamingTokens->GetExternalTokensChecked(GetNamingTokensData()->NamingTokensExternalGuid);
+
+		ExternalTokens.Reset(GetNamingTokensData()->UserDefinedTokens.Num());
+		for (const TTuple<FNamingTokenData, FText>& UserToken : GetNamingTokensData()->UserDefinedTokens)
+		{
+			ExternalTokens.Add({
+				UserToken.Key.TokenKey,
+				UserToken.Key.DisplayName,
+				// @todo naming tokens - If we support TOptional<FText> that can be used instead of a lambda
+				FNamingTokenData::FTokenProcessorDelegateNative::CreateLambda([Value = UserToken.Value]() {
+					return Value;
+				})
+			});
+		}
+	}
+}
+
+void UTakeRecorderSubsystemImplementation::OnTakePresetSettingsChanged()
+{
+	// This effectively regenerates the level sequence in the preset.
+	// Existing data is discarded. In the future, we should try to migrate it.
+	// FTakePresetRecorderCustomization handles asking the user whether it is ok to discard the changes. Once this event fires, the decision has been made.
+	if (!GIsTransacting)
+	{
+		ClearPendingTake();
+	}
+}
+
+void UTakeRecorderSubsystemImplementation::OnPreForceDeleteObjects(const TArray<UObject*>& InObjects)
+{
+	// Take meta data is cached as a strong reference, but can prevent us from deleting owning level sequences.
+	// We need to keep a strong reference to it for Take Recorder functionality, especially in the panel, and it should persist
+	// even with the panel closed. Listening for a force delete allows us to safely clear the reference.
+	const ULevelSequence* LevelSequence = GetLevelSequence();
+	if (!LevelSequence && LastRecordedLevelSequence.IsValid())
+	{
+		LevelSequence = LastRecordedLevelSequence.Get();
+	}
+	if (LevelSequence && TakeMetaData && InObjects.Contains(LevelSequence)
+		&& LevelSequence->FindMetaData<UTakeMetaData>() == TakeMetaData)
+	{
+		if (IsRecording())
+		{
+			// Need to cancel or we will be in a bad state. TakeRecorder will tick but be without a level sequence.
+			CancelRecording(false);
+		}
+		TakeMetaData = nullptr;
+		FEditorDelegates::OnPreForceDeleteObjects.Remove(OnPreForceDeleteObjectsHandle);
+		OnPreForceDeleteObjectsHandle.Reset();
+	}
+}
+
+#undef WITH_OWNING_SUBSYSTEM
+#undef LOCTEXT_NAMESPACE
\ No newline at end of file
diff --git a/Engine/Plugins/VirtualProduction/Takes/Source/TakeRecorder/Private/Recorder/TakeRecorderSubsystemImplementation.h b/Engine/Plugins/VirtualProduction/Takes/Source/TakeRecorder/Private/Recorder/TakeRecorderSubsystemImplementation.h
new file mode 100644
--- /dev/null
+++ b/Engine/Plugins/VirtualProduction/Takes/Source/TakeRecorder/Private/Recorder/TakeRecorderSubsystemImplementation.h
@@ -0,0 +1,187 @@
+// Copyright Epic Games, Inc. All Rights Reserved.
+
+#pragma once
+
+#include "CoreMinimal.h"
+#include "Recorder/ITakeRecorderSubsystemInterface.h"
+#include "Recorder/TakeRecorderParameters.h"
+#include "Tickable.h"
+
+#include "TakeRecorderSubsystemImplementation.generated.h"
+
+class UTakeRecorder;
+class UTakeRecorderSource;
+class UTakeRecorderSubsystem;
+struct FNamingTokensEvaluationData;
+
+UCLASS()
+class UTakeRecorderSubsystemImplementation final : public UObject, public ITakeRecorderSubsystemInterface, public FTickableGameObject
+{
+	GENERATED_BODY()
+
+public:
+	// ~Begin FTickableGameObject
+	virtual ETickableTickType GetTickableTickType() const override;
+	virtual UWorld* GetTickableGameObjectWorld() const override;
+	virtual bool IsTickable() const override;
+	virtual bool IsTickableWhenPaused() const override;
+	virtual bool IsTickableInEditor() const override;
+	virtual void Tick(float DeltaTime) override;
+	virtual TStatId GetStatId() const override;
+	// ~End FTickableGameObject
+
+	/** Perform implementation specific initialization. */
+	void InitializeImplementation(UTakeRecorderSubsystem* OwningSubsystem);
+
+	/** Perform implementation specific shutdown. */
+	void DeinitializeImplementation();
+
+	// ~Begin ITakeRecorderSubsystem interface
+	virtual void SetTargetSequence(const FTakeRecorderSequenceParameters& InData = FTakeRecorderSequenceParameters()) override;
+	virtual void SetRecordIntoLevelSequence(ULevelSequence* LevelSequence) override;
+	virtual bool CanReviewLastRecording() const override;
+	virtual bool ReviewLastRecording() override;
+	virtual bool StartRecording(bool bOpenSequencer = true, bool bShowErrorMessage = true) override;
+	virtual void StopRecording() override;
+	virtual void CancelRecording(bool bShowConfirmMessage = true) override;
+	virtual void ResetToPendingTake() override;
+	virtual void ClearPendingTake() override;
+	virtual UTakePreset* GetPendingTake() const override;
+	virtual void RevertChanges() override;
+	virtual UTakeRecorderSource* AddSource(const TSubclassOf<UTakeRecorderSource> InSourceClass) override;
+	virtual void RemoveSource(UTakeRecorderSource* InSource) override;
+	virtual void ClearSources() override;
+	virtual UTakeRecorderSources* GetSources() const override;
+	virtual TArrayView<UTakeRecorderSource* const> GetAllSources() const override;
+	virtual TArray<UTakeRecorderSource*> GetAllSourcesCopy() const override;
+	virtual UTakeRecorderSource* GetSourceByClass(const TSubclassOf<UTakeRecorderSource> InSourceClass) const override;
+	virtual void AddSourceForActor(AActor* InActor, bool bReduceKeys = true, bool bShowProgress = true) override;
+	virtual void RemoveActorFromSources(AActor* InActor) override;
+	virtual AActor* GetSourceActor(UTakeRecorderSource* InSource) const override;
+	virtual ETakeRecorderState GetState() const override;
+	virtual void SetTakeNumber(int32 InNewTakeNumber, bool bEmitChanged = true) override;
+	virtual int32 GetNextTakeNumber(const FString& InSlate) const override;
+	virtual void GetNumberOfTakes(const FString& InSlate, int32& OutMaxTake, int32& OutNumTakes) const override;
+	virtual TArray<FAssetData> GetSlates(FName InPackagePath = NAME_None) const override;
+	virtual void SetSlateName(const FString& InSlateName, bool bEmitChanged = true) override;
+	virtual bool MarkFrame() override;
+	virtual FFrameRate GetFrameRate() const override;
+	virtual void SetFrameRate(FFrameRate InFrameRate) override;
+	virtual void SetFrameRateFromTimecode() override;
+	virtual void ImportPreset(const FAssetData& InPreset) override;
+	virtual bool IsReviewing() const override;
+	virtual bool IsRecording() const override;
+	virtual bool TryGetSequenceCountdown(float& OutValue) const override;
+	virtual void SetSequenceCountdown(float InSeconds) override;
+	virtual TArray<UObject*> GetSourceRecordSettings(UTakeRecorderSource* InSource) const override;
+	virtual FTakeRecorderParameters GetGlobalRecordSettings() const override;
+	virtual void SetGlobalRecordSettings(const FTakeRecorderParameters& InParameters) override;
+	virtual UTakeMetaData* GetTakeMetaData() const override;
+	virtual ULevelSequence* GetLevelSequence() const override;
+	virtual ULevelSequence* GetSuppliedLevelSequence() const override;
+	virtual ULevelSequence* GetRecordingLevelSequence() const override;
+	virtual ULevelSequence* GetRecordIntoLevelSequence() const override;
+	virtual ULevelSequence* GetLastRecordedLevelSequence() const override;
+	virtual UTakePreset* GetTransientPreset() const override;
+	virtual ETakeRecorderMode GetTakeRecorderMode() const override;
+	virtual UTakeRecorderNamingTokensData* GetNamingTokensData() const override;
+	virtual bool HasPendingChanges() const override;
+	// ~End ITakeRecorderSubsystem interface
+	
+private:
+	/**
+	 * Allocate the preset required for interacting with this subsystem. Re-uses an existing preset if necessary.
+	 */
+	static UTakePreset* AllocateTransientPreset();
+
+	/** Cache the current available metadata. */
+	void CacheMetaData();
+
+	/** Sets the frame rate. */
+	void SetFrameRateImpl(const FFrameRate& InFrameRate, bool bFromTimecode);
+
+	/** Calculate and apply the next take number. */
+	void IncrementTakeNumber();
+
+	void OnAssetRegistryFilesLoaded();
+	void OnRecordingInitialized(UTakeRecorder* Recorder);
+	void OnRecordingStarted(UTakeRecorder* Recorder);
+	void OnRecordingStopped(UTakeRecorder* Recorder);
+	void OnRecordingFinished(UTakeRecorder* Recorder);
+	void OnRecordingCancelled(UTakeRecorder* Recorder);
+	void OnTakeSlateChanged(const FString& InSlate, UTakeMetaData* InTakeMetaData);
+	void OnTakeNumberChanged(int32 InTakeNumber, UTakeMetaData* InTakeMetaData);
+	
+	/** Update the last level sequence. */
+	void SetLastLevelSequence(ULevelSequence *InSequence);
+
+	/** Callback when any source is added. */
+	void OnSourceAdded(UTakeRecorderSource* Source);
+
+	/** Callback when any source is removed. */
+	void OnSourceRemoved(UTakeRecorderSource* Source);
+
+	/** Setup bindings to our naming tokens. */
+	void BindToNamingTokenEvents();
+	
+	/** Remove bindings to our naming tokens. */
+	void UnbindNamingTokensEvents();
+
+	/** Called before our naming tokens class evaluates. Used to populate naming token data with our custom token definitions. */
+	void OnTakeRecorderNamingTokensPreEvaluate(const FNamingTokensEvaluationData& InEvaluationData);
+	
+	/** Handles the UTakePresetSettings::RecordTargetClass changing: Recreates the transient level sequence if recording a transaction. */
+	void OnTakePresetSettingsChanged();
+
+	/** Callback when the engine is attempting to force delete objects. */
+	void OnPreForceDeleteObjects(const TArray<UObject*>& InObjects);
+	
+private:
+	/** The public facing engine subsystem we are implementing. */
+	TWeakObjectPtr<UTakeRecorderSubsystem> OwningSubsystemWeakPtr;
+	
+	/** Last data used for initialization. */
+	FTakeRecorderSequenceParameters TargetSequenceData;
+
+	/** A transient preset available for the subsystem. */
+	UPROPERTY(Transient)
+	TObjectPtr<UTakePreset> TransientPreset;
+
+	/** Current supplied level sequence. */
+	UPROPERTY()
+	TObjectPtr<ULevelSequence> SuppliedLevelSequence;
+
+	/** Current record into level sequence. */
+	UPROPERTY()
+	TObjectPtr<ULevelSequence> RecordIntoLevelSequence;
+
+	/** Current recording level sequence. */
+	UPROPERTY()
+	TObjectPtr<ULevelSequence> RecordingLevelSequence;
+
+	/** The last recorded level sequence. */
+	UPROPERTY()
+	TWeakObjectPtr<ULevelSequence> LastRecordedLevelSequence;
+
+	/** Take meta-data cached from the level sequence if it exists. */
+	UPROPERTY()
+	TObjectPtr<UTakeMetaData> TakeMetaData;
+	
+	/** Transient take meta-data owned by this subsystem. Only used if none exists on the level sequence already. */
+	UPROPERTY(Transient)
+	TObjectPtr<UTakeMetaData> TransientTakeMetaData;
+
+	/** Stored data relevant to our naming tokens. Managed as UObject for transactions. */
+	UPROPERTY(Transient)
+	TObjectPtr<UTakeRecorderNamingTokensData> NamingTokensData;
+	
+	FDelegateHandle OnAssetRegistryFilesLoadedHandle;
+	FDelegateHandle OnRecordingInitializedHandle, OnRecordingStoppedHandle, OnRecordingFinishedHandle, OnRecordingCancelledHandle;
+	FDelegateHandle OnPreForceDeleteObjectsHandle;
+	
+	/** Whether we should auto apply the next available take number when asset discovery has finished or not */
+	bool bAutoApplyTakeNumber = false;
+
+	/** If SetTargetSequence has been called, and we're fully initialized. */
+	bool bHasTargetSequenceBeenSet = false;
+};
diff --git a/Engine/Plugins/VirtualProduction/Takes/Source/TakeRecorder/Public/Recorder/ITakeRecorderSubsystemInterface.h b/Engine/Plugins/VirtualProduction/Takes/Source/TakeRecorder/Public/Recorder/ITakeRecorderSubsystemInterface.h
new file mode 100644
--- /dev/null
+++ b/Engine/Plugins/VirtualProduction/Takes/Source/TakeRecorder/Public/Recorder/ITakeRecorderSubsystemInterface.h
@@ -0,0 +1,82 @@
+// Copyright Epic Games, Inc. All Rights Reserved.
+
+#pragma once
+
+#include "AssetRegistry/AssetData.h"
+#include "TakeRecorderParameters.h"
+#include "UObject/Interface.h"
+
+#include "ITakeRecorderSubsystemInterface.generated.h"
+
+class UTakeRecorderNamingTokensData;
+class UTakeRecorderSources;
+class UTakeRecorderSource;
+class UTakeMetaData;
+enum class ETakeRecorderState : uint8;
+
+UINTERFACE(MinimalAPI, NotBlueprintable)
+class UTakeRecorderSubsystemInterface : public UInterface
+{
+	GENERATED_BODY()
+};
+
+/** Interface for the public Take Recorder Subsystem and the private implementation. */
+class ITakeRecorderSubsystemInterface
+{
+	GENERATED_BODY()
+	
+public:
+	virtual void SetTargetSequence(const FTakeRecorderSequenceParameters& InData = FTakeRecorderSequenceParameters()) = 0;
+	virtual void SetRecordIntoLevelSequence(ULevelSequence* LevelSequence) = 0;
+	virtual bool CanReviewLastRecording() const = 0;
+	virtual bool ReviewLastRecording() = 0;
+	virtual bool StartRecording(bool bOpenSequencer = true, bool bShowErrorMessage = true) = 0;
+	virtual void StopRecording() = 0;
+	virtual void CancelRecording(bool bShowConfirmMessage = true) = 0;
+	virtual void ResetToPendingTake() = 0;
+	virtual void ClearPendingTake() = 0;
+	virtual UTakePreset* GetPendingTake() const = 0;
+	virtual void RevertChanges() = 0;
+	
+	// ~Begin Sources
+	virtual UTakeRecorderSource* AddSource(const TSubclassOf<UTakeRecorderSource> InSourceClass) = 0;
+	virtual void RemoveSource(UTakeRecorderSource* InSource) = 0;
+	virtual void ClearSources() = 0;
+	virtual UTakeRecorderSources* GetSources() const = 0;
+	virtual TArrayView<UTakeRecorderSource* const> GetAllSources() const = 0;
+	virtual TArray<UTakeRecorderSource*> GetAllSourcesCopy() const = 0;
+	virtual UTakeRecorderSource* GetSourceByClass(const TSubclassOf<UTakeRecorderSource> InSourceClass) const = 0;
+	virtual void AddSourceForActor(AActor* InActor, bool bReduceKeys = true, bool bShowProgress = true) = 0;
+	virtual void RemoveActorFromSources(AActor* InActor) = 0;
+	virtual AActor* GetSourceActor(UTakeRecorderSource* InSource) const = 0;
+	// ~End Sources
+	
+	virtual ETakeRecorderState GetState() const = 0;
+	virtual void SetTakeNumber(int32 InNewTakeNumber, bool bEmitChanged = true) = 0;
+	virtual int32 GetNextTakeNumber(const FString& InSlate) const = 0;
+	virtual void GetNumberOfTakes(const FString& InSlate, int32& OutMaxTake, int32& OutNumTakes) const = 0;
+	virtual TArray<FAssetData> GetSlates(FName InPackagePath = NAME_None) const = 0;
+	virtual void SetSlateName(const FString& InSlateName, bool bEmitChanged = true) = 0;
+	virtual bool MarkFrame() = 0;
+	virtual FFrameRate GetFrameRate() const = 0;
+	virtual void SetFrameRate(FFrameRate InFrameRate) = 0;
+	virtual void SetFrameRateFromTimecode() = 0;
+	virtual void ImportPreset(const FAssetData& InPreset) = 0;
+	virtual bool IsReviewing() const = 0;
+	virtual bool IsRecording() const = 0;
+	virtual bool TryGetSequenceCountdown(float& OutValue) const = 0;
+	virtual void SetSequenceCountdown(float InSeconds) = 0;
+	virtual TArray<UObject*> GetSourceRecordSettings(UTakeRecorderSource* InSource) const = 0;
+	virtual FTakeRecorderParameters GetGlobalRecordSettings() const = 0;
+	virtual void SetGlobalRecordSettings(const FTakeRecorderParameters& InParameters) = 0;
+	virtual UTakeMetaData* GetTakeMetaData() const = 0;
+	virtual ULevelSequence* GetLevelSequence() const = 0;
+	virtual ULevelSequence* GetSuppliedLevelSequence() const = 0;
+	virtual ULevelSequence* GetRecordingLevelSequence() const = 0;
+	virtual ULevelSequence* GetRecordIntoLevelSequence() const = 0;
+	virtual ULevelSequence* GetLastRecordedLevelSequence() const = 0;
+	virtual UTakePreset* GetTransientPreset() const = 0;
+	virtual ETakeRecorderMode GetTakeRecorderMode() const = 0;
+	virtual UTakeRecorderNamingTokensData* GetNamingTokensData() const = 0;
+	virtual bool HasPendingChanges() const = 0;
+};
diff --git a/Engine/Plugins/VirtualProduction/Takes/Source/TakeRecorder/Public/Recorder/TakeRecorderSubsystem.h b/Engine/Plugins/VirtualProduction/Takes/Source/TakeRecorder/Public/Recorder/TakeRecorderSubsystem.h
--- a/Engine/Plugins/VirtualProduction/Takes/Source/TakeRecorder/Public/Recorder/TakeRecorderSubsystem.h
+++ b/Engine/Plugins/VirtualProduction/Takes/Source/TakeRecorder/Public/Recorder/TakeRecorderSubsystem.h
@@ -2,10 +2,10 @@
 
 #pragma once
 
+#include "ITakeRecorderSubsystemInterface.h"
 #include "Subsystems/EngineSubsystem.h"
 #include "TakeRecorder.h"
 #include "TakeRecorderSources.h"
-#include "Tickable.h"
 
 #include "TakeRecorderSubsystem.generated.h"
 
@@ -38,7 +38,7 @@ DECLARE_DYNAMIC_MULTICAST_DELEGATE_OneParam(FOnTakeRecorderSourceRemovedDynamic,
 * and register the subsystem as tickable currently required to cache level metadata.
 */
 UCLASS()
-class TAKERECORDER_API UTakeRecorderSubsystem : public UEngineSubsystem, public FTickableGameObject
+class TAKERECORDER_API UTakeRecorderSubsystem final : public UEngineSubsystem, public ITakeRecorderSubsystemInterface
 {
 	GENERATED_BODY()
 
@@ -47,15 +47,7 @@ public:
 	virtual void Initialize(FSubsystemCollectionBase& Collection) override;
 	virtual void Deinitialize() override;
 
-	// Begin FTickableGameObject
-	virtual ETickableTickType GetTickableTickType() const override;
-	virtual UWorld* GetTickableGameObjectWorld() const override;
-	virtual bool IsTickable() const override;
-	virtual bool IsTickableWhenPaused() const override;
-	virtual bool IsTickableInEditor() const override;
-	virtual void Tick(float DeltaTime) override;
-	virtual TStatId GetStatId() const override;
-	// ~End FTickableGameObject
+	// ~Begin ITakeRecorderSubsystem interface
 	
 	/**
 	 * Provide sequence data for this take recorder. This will also perform initialization of the subsystem.
@@ -64,21 +56,21 @@ public:
 	 * @param InData The sequence parameters, which are mutually exclusive. They can all be null.
 	 */
 	UFUNCTION(BlueprintCallable, Category = TakeRecorder, meta=(AutoCreateRefTerm="InData"))
-	void SetTargetSequence(const FTakeRecorderSequenceParameters& InData = FTakeRecorderSequenceParameters());
+	virtual void SetTargetSequence(const FTakeRecorderSequenceParameters& InData = FTakeRecorderSequenceParameters()) override;
 	
 	/** Set the record into level sequence. */
-	void SetRecordIntoLevelSequence(ULevelSequence* LevelSequence);
+	virtual void SetRecordIntoLevelSequence(ULevelSequence* LevelSequence) override;
 
 	/** Can we review the last recording? */
 	UFUNCTION(BlueprintCallable, Category = TakeRecorder)
-	bool CanReviewLastRecording() const;
+	virtual bool CanReviewLastRecording() const override;
 	
 	/**
 	 * Supply the last recording if it exists.
 	 * @return true if it can be reviewed, false if it there isn't a recording to review.
 	 */
 	UFUNCTION(BlueprintCallable, Category = TakeRecorder)
-	bool ReviewLastRecording();
+	virtual bool ReviewLastRecording() override;
 
 	/**
 	 * Begin a new recording.
@@ -87,59 +79,59 @@ public:
 	 * @return true if the recording was started successfully.
 	 */
 	UFUNCTION(BlueprintCallable, Category = TakeRecorder)
-	bool StartRecording(bool bOpenSequencer = true, bool bShowErrorMessage = true);
+	virtual bool StartRecording(bool bOpenSequencer = true, bool bShowErrorMessage = true) override;
 
 	/**
 	 * Stop an existing recording.
 	 */
 	UFUNCTION(BlueprintCallable, Category = TakeRecorder)
-	void StopRecording();
+	virtual void StopRecording() override;
 
 	/**
 	 * Cancel an in-progress recording.
 	 */
 	UFUNCTION(BlueprintCallable, Category = TakeRecorder)
-	void CancelRecording(bool bShowConfirmMessage = true);
+	virtual void CancelRecording(bool bShowConfirmMessage = true) override;
 
 	/**
 	 * Reset to the pending take.
 	 */
 	UFUNCTION(BlueprintCallable, Category = TakeRecorder)
-	void ResetToPendingTake();
+	virtual void ResetToPendingTake() override;
 
 	/**
 	 * Clear the pending take.
 	 */
 	UFUNCTION(BlueprintCallable, Category = TakeRecorder)
-	void ClearPendingTake();
+	virtual void ClearPendingTake() override;
 
 	/** Retrieve the pending take. This may be null. */
 	UFUNCTION(BlueprintCallable, Category = TakeRecorder)
-	UTakePreset* GetPendingTake() const;
+	virtual UTakePreset* GetPendingTake() const override;
 
 	/** Revert any changes restoring the preset origin. */
 	UFUNCTION(BlueprintCallable, Category = TakeRecorder)
-	void RevertChanges();
+	virtual void RevertChanges() override;
 	
 	// ~Begin Sources
 	/** Add a source by a source class. */
 	UFUNCTION(BlueprintCallable, Category = TakeRecorder, meta = (DeterminesOutputType = "InSourceClass"))
-	UTakeRecorderSource* AddSource(const TSubclassOf<UTakeRecorderSource> InSourceClass);
+	virtual UTakeRecorderSource* AddSource(const TSubclassOf<UTakeRecorderSource> InSourceClass) override;
 	
 	/** Remove a given source. */
 	UFUNCTION(BlueprintCallable, Category = TakeRecorder)
-	void RemoveSource(UTakeRecorderSource* InSource);
+	virtual void RemoveSource(UTakeRecorderSource* InSource) override;
 
 	/** Remove all sources from the current sequence. */
 	UFUNCTION(BlueprintCallable, Category = TakeRecorder)
-	void ClearSources();
+	virtual void ClearSources() override;
 
 	/** Retrieve the sources. */
 	UFUNCTION(BlueprintCallable, Category = TakeRecorder)
-	UTakeRecorderSources* GetSources() const;
+	virtual UTakeRecorderSources* GetSources() const override;
 	
 	/** Retrieve all sources for the current sequence. */
-	TArrayView<UTakeRecorderSource* const> GetAllSources() const;
+	virtual TArrayView<UTakeRecorderSource* const> GetAllSources() const override;
 
 	/**
 	 * Retrieves a copy of the list of sources that are being recorded. This is intended for Blueprint usages which cannot
@@ -147,11 +139,11 @@ public:
 	 * DO NOT MODIFY THIS ARRAY, modifications will be lost.
 	 */
 	UFUNCTION(BlueprintPure, DisplayName = "Get Sources (Copy)", Category = TakeRecorder)
-	TArray<UTakeRecorderSource*> GetAllSourcesCopy() const;
+	virtual TArray<UTakeRecorderSource*> GetAllSourcesCopy() const override;
 
 	/** Retrieve the first source of the given class. */
 	UFUNCTION(BlueprintCallable, Category = TakeRecorder, meta = (DeterminesOutputType = "InSourceClass"))
-	UTakeRecorderSource* GetSourceByClass(const TSubclassOf<UTakeRecorderSource> InSourceClass) const;
+	virtual UTakeRecorderSource* GetSourceByClass(const TSubclassOf<UTakeRecorderSource> InSourceClass) const override;
 
 	/*
 	 * Add an actor as a source.
@@ -161,22 +153,22 @@ public:
 	 * @param bShowProgress Enable/disable the dialog box showing progress for the potentially slow parts of finalizing the take
 	 */
 	UFUNCTION(BlueprintCallable, Category = TakeRecorder)
-	void AddSourceForActor(AActor* InActor, bool bReduceKeys = true, bool bShowProgress = true);
+	virtual void AddSourceForActor(AActor* InActor, bool bReduceKeys = true, bool bShowProgress = true) override;
 
 	/**
 	 * Remove an actor from available sources.
 	 */
 	UFUNCTION(BlueprintCallable, Category = TakeRecorder)
-	void RemoveActorFromSources(AActor* InActor);
+	virtual void RemoveActorFromSources(AActor* InActor) override;
 
 	/** Retrieve the actor from a source, if applicable. */
 	UFUNCTION(BlueprintCallable, Category = TakeRecorder)
-	AActor* GetSourceActor(UTakeRecorderSource* InSource) const;
+	virtual AActor* GetSourceActor(UTakeRecorderSource* InSource) const override;
 	// ~End Sources
 	
 	/** Retrieve the current take recorder state. */
 	UFUNCTION(BlueprintCallable, Category = TakeRecorder)
-	ETakeRecorderState GetState() const;
+	virtual ETakeRecorderState GetState() const override;
 
 	/**
 	 * Directly set the take number.
@@ -184,22 +176,22 @@ public:
 	 * @param bEmitChanged Whether to broadcast events signaling the take number has changed.
 	 */
 	UFUNCTION(BlueprintCallable, Category = TakeRecorder)
-	void SetTakeNumber(int32 InNewTakeNumber, bool bEmitChanged = true);
+	virtual void SetTakeNumber(int32 InNewTakeNumber, bool bEmitChanged = true) override;
 
 	/** Compute the next take number given a slate. */
 	UFUNCTION(BlueprintCallable, Category = TakeRecorder)
-	int32 GetNextTakeNumber(const FString& InSlate) const;
+	virtual int32 GetNextTakeNumber(const FString& InSlate) const override;
 
 	/** Find both the current maximum take value and the total number of takes for a given slate. */
 	UFUNCTION(BlueprintCallable, Category = TakeRecorder)
-	void GetNumberOfTakes(const FString& InSlate, int32& OutMaxTake, int32& OutNumTakes) const;
+	virtual void GetNumberOfTakes(const FString& InSlate, int32& OutMaxTake, int32& OutNumTakes) const override;
 
 	/**
 	 * Retrieve all slates.
 	 * @param InPackagePath [Optional] Scope the search to a specific folder, recursively.
 	 */
 	UFUNCTION(BlueprintCallable, Category = TakeRecorder)
-	TArray<FAssetData> GetSlates(FName InPackagePath = NAME_None) const;
+	virtual TArray<FAssetData> GetSlates(FName InPackagePath = NAME_None) const override;
 	
 	/**
 	 * Directly set the slate name.
@@ -207,40 +199,40 @@ public:
 	 * @param bEmitChanged Whether to broadcast events signaling the slate name has changed.
 	 */
 	UFUNCTION(BlueprintCallable, Category = TakeRecorder)
-	void SetSlateName(const FString& InSlateName, bool bEmitChanged = true);
+	virtual void SetSlateName(const FString& InSlateName, bool bEmitChanged = true) override;
 
 	/** Mark the current frame. */
 	UFUNCTION(BlueprintCallable, Category = TakeRecorder)
-	bool MarkFrame();
+	virtual bool MarkFrame() override;
 
 	/**
 	 * Access the frame rate for this take
 	 */
 	UFUNCTION(BlueprintCallable, Category = TakeRecorder)
-	FFrameRate GetFrameRate() const;
+	virtual FFrameRate GetFrameRate() const override;
 	
 	/**
 	 * Set the frame rate for this take
 	 */
 	UFUNCTION(BlueprintCallable, Category = TakeRecorder)
-	void SetFrameRate(FFrameRate InFrameRate);
+	virtual void SetFrameRate(FFrameRate InFrameRate) override;
 	
 	/**
 	 * Set if the frame rate is set from the Timecode frame rate
 	 */
 	UFUNCTION(BlueprintCallable, Category = TakeRecorder)
-	void SetFrameRateFromTimecode();
+	virtual void SetFrameRateFromTimecode() override;
 
 	/** Import a preset to the transient preset. */
-	void ImportPreset(const FAssetData& InPreset);
+	virtual void ImportPreset(const FAssetData& InPreset) override;
 
 	/** If Take Recorder is currently reviewing. */
 	UFUNCTION(BlueprintCallable, Category = TakeRecorder)
-	bool IsReviewing() const;
+	virtual bool IsReviewing() const override;
 
 	/** If Take Recorder is currently recording. */
 	UFUNCTION(BlueprintCallable, Category = TakeRecorder)
-	bool IsRecording() const;
+	virtual bool IsRecording() const override;
 
 	/**
 	 * Retrieve the current sequence's countdown.
@@ -248,60 +240,60 @@ public:
 	 * @return True if we are in a countdown sequence, false otherwise.
 	 */
 	UFUNCTION(BlueprintCallable, Category = TakeRecorder)
-	bool TryGetSequenceCountdown(float& OutValue) const;
+	virtual bool TryGetSequenceCountdown(float& OutValue) const override;
 
 	/**
 	 * Sets the current sequence's countdown.
 	 * @param InSeconds Time in seconds.
 	 */
 	UFUNCTION(BlueprintCallable, Category = TakeRecorder)
-	void SetSequenceCountdown(float InSeconds);
+	virtual void SetSequenceCountdown(float InSeconds) override;
 
 	/** Retrieve additional settings objects from a source. */
 	UFUNCTION(BlueprintCallable, Category = TakeRecorder)
-	TArray<UObject*> GetSourceRecordSettings(UTakeRecorderSource* InSource) const;
+	virtual TArray<UObject*> GetSourceRecordSettings(UTakeRecorderSource* InSource) const override;
 
 	/** Retrieve the global take recorder settings. */
 	UFUNCTION(BlueprintCallable, Category = TakeRecorder)
-	FTakeRecorderParameters GetGlobalRecordSettings() const;
+	virtual FTakeRecorderParameters GetGlobalRecordSettings() const override;
 
 	/** Set the global take recorder settings. */
 	UFUNCTION(BlueprintCallable, Category = TakeRecorder)
-	void SetGlobalRecordSettings(const FTakeRecorderParameters& InParameters);
+	virtual void SetGlobalRecordSettings(const FTakeRecorderParameters& InParameters) override;
 	
 	/** Retrieve the current meta data. */
 	UFUNCTION(BlueprintCallable, Category = TakeRecorder)
-	UTakeMetaData* GetTakeMetaData() const;
+	virtual UTakeMetaData* GetTakeMetaData() const override;
 
 	/** Return the level sequence we are using. */
 	UFUNCTION(BlueprintCallable, Category = TakeRecorder)
-	ULevelSequence* GetLevelSequence() const;
+	virtual ULevelSequence* GetLevelSequence() const override;
 
 	/** Retrieve the Supplied Level Sequence if it exists. */
-	ULevelSequence* GetSuppliedLevelSequence() const;
+	virtual ULevelSequence* GetSuppliedLevelSequence() const override;
 
 	/** Retrieve the Recording Level Sequence if it exists. */
-	ULevelSequence* GetRecordingLevelSequence() const;
+	virtual ULevelSequence* GetRecordingLevelSequence() const override;
 
 	/** Retrieve the Record Into Level Sequence if it exists. */
-	ULevelSequence* GetRecordIntoLevelSequence() const;
+	virtual ULevelSequence* GetRecordIntoLevelSequence() const override;
 
 	/** Retrieve the Last Recorded Level Sequence if it exists. */
 	UFUNCTION(BlueprintCallable, Category = TakeRecorder)
-	ULevelSequence* GetLastRecordedLevelSequence() const;
+	virtual ULevelSequence* GetLastRecordedLevelSequence() const override;
 
 	/** Retrieve the transient preset the subsystem is using. */
-	UTakePreset* GetTransientPreset() const;
+	virtual UTakePreset* GetTransientPreset() const override;
 	
 	/** The current take recorder mode. */
 	UFUNCTION(BlueprintCallable, Category = TakeRecorder)
-	ETakeRecorderMode GetTakeRecorderMode() const;
+	virtual ETakeRecorderMode GetTakeRecorderMode() const override;
 
 	/** Retrieve relevant Naming Tokens data for Take Recorder. */
-	UTakeRecorderNamingTokensData* GetNamingTokensData() const;
+	virtual UTakeRecorderNamingTokensData* GetNamingTokensData() const override;
 
 	/** @return Whether there are any pending changes, which ClearPendingTake would discard. */
-	bool HasPendingChanges() const;
+	virtual bool HasPendingChanges() const override;
 	
 	/**
 	 * Retrieve a multi-cast delegate that is triggered when a recording pre-initializes
@@ -335,90 +327,16 @@ public:
 	 * Retrieve a multi-cast delegate that is triggered when a source is removed.
 	 */
 	UTakeRecorderSources::FOnSourceRemoved& GetOnRecordingSourceRemovedEvent();
-	
-private:
-	/**
-	 * Allocate the preset required for interacting with this subsystem. Re-uses an existing preset if necessary
-	 */
-	static UTakePreset* AllocateTransientPreset();
-
-	/** Cache the current available metadata. */
-	void CacheMetaData();
-
-	/** Sets the frame rate. */
-	void SetFrameRateImpl(const FFrameRate& InFrameRate, bool bFromTimecode);
-
-	/** Calculate and apply the next take number. */
-	void IncrementTakeNumber();
-
-	void OnAssetRegistryFilesLoaded();
-	void OnRecordingInitialized(UTakeRecorder* Recorder);
-	void OnRecordingStarted(UTakeRecorder* Recorder);
-	void OnRecordingStopped(UTakeRecorder* Recorder);
-	void OnRecordingFinished(UTakeRecorder* Recorder);
-	void OnRecordingCancelled(UTakeRecorder* Recorder);
-	void OnTakeSlateChanged(const FString& InSlate, UTakeMetaData* InTakeMetaData);
-	void OnTakeNumberChanged(int32 InTakeNumber, UTakeMetaData* InTakeMetaData);
-	
-	/** Update the last level sequence. */
-	void SetLastLevelSequence(ULevelSequence *InSequence);
-
-	/** Callback when any source is added. */
-	void OnSourceAdded(UTakeRecorderSource* Source);
-
-	/** Callback when any source is removed. */
-	void OnSourceRemoved(UTakeRecorderSource* Source);
-
-	/** Setup bindings to our naming tokens. */
-	void BindToNamingTokenEvents();
-	
-	/** Remove bindings to our naming tokens. */
-	void UnbindNamingTokensEvents();
-
-	/** Called before our naming tokens class evaluates. Used to populate naming token data with our custom token definitions. */
-	void OnTakeRecorderNamingTokensPreEvaluate(const FNamingTokensEvaluationData& InEvaluationData);
-	
-	/** Handles the UTakePresetSettings::RecordTargetClass changing: Recreates the transient level sequence if recording a transaction. */
-	void OnTakePresetSettingsChanged();
 
-	/** Callback when the engine is attempting to force delete objects. */
-	void OnPreForceDeleteObjects(const TArray<UObject*>& InObjects);
+	// ~End ITakeRecorderSubsystem interface
 	
 private:
-	/** Last data used for initialization. */
-	FTakeRecorderSequenceParameters TargetSequenceData;
-
-	/** A transient preset available for the subsystem. */
-	UPROPERTY(Transient)
-	TObjectPtr<UTakePreset> TransientPreset;
-
-	/** Current supplied level sequence. */
-	UPROPERTY()
-	TObjectPtr<ULevelSequence> SuppliedLevelSequence;
-
-	/** Current record into level sequence. */
-	UPROPERTY()
-	TObjectPtr<ULevelSequence> RecordIntoLevelSequence;
-
-	/** Current recording level sequence. */
-	UPROPERTY()
-	TObjectPtr<ULevelSequence> RecordingLevelSequence;
-
-	/** The last recorded level sequence. */
-	UPROPERTY()
-	TWeakObjectPtr<ULevelSequence> LastRecordedLevelSequence;
-
-	/** Take meta-data cached from the level sequence if it exists. */
-	UPROPERTY()
-	TObjectPtr<UTakeMetaData> TakeMetaData;
+	// Friend class declared as a precaution in case the implementation needs to access private members.
+	friend class UTakeRecorderSubsystemImplementation;
 	
-	/** Transient take meta-data owned by this subsystem. Only used if none exists on the level sequence already. */
-	UPROPERTY(Transient)
-	TObjectPtr<UTakeMetaData> TransientTakeMetaData;
-
-	/** Stored data relevant to our naming tokens. Managed as UObject for transactions. */
-	UPROPERTY(Transient)
-	TObjectPtr<UTakeRecorderNamingTokensData> NamingTokensData;
+	/** The actual Take Recorder Subsystem implementation. */
+	UPROPERTY()
+	TScriptInterface<ITakeRecorderSubsystemInterface> Implementation;
 	
 	// These native delegates are declared elsewhere throughout Take Recorder and we reuse them here. However, it's important
 	// the editor subsystem manage its own variants because the order they fire is important for listeners, such as the panel. For example, the
@@ -442,19 +360,8 @@ private:
 	UTakeRecorderSources::FOnSourceAdded OnRecordingSourceAddedEvent;
 	/** Multicast delegate when any source has been added. For native binding. */
 	UTakeRecorderSources::FOnSourceRemoved OnRecordingSourceRemovedEvent;
-	
-	FDelegateHandle OnAssetRegistryFilesLoadedHandle;
-	FDelegateHandle OnRecordingInitializedHandle, OnRecordingStoppedHandle, OnRecordingFinishedHandle, OnRecordingCancelledHandle;
-	FDelegateHandle OnPreForceDeleteObjectsHandle;
-	
-	/** Whether we should auto apply the next available take number when asset discovery has finished or not */
-	bool bAutoApplyTakeNumber;
 
-	/** If SetTargetSequence has been called, and we're fully initialized. */
-	bool bHasTargetSequenceBeenSet = false;
-	
 public:
-
 	/** Called before initialization occurs (ie. when the recording button is pressed and before the countdown starts) */
 	UPROPERTY(BlueprintAssignable, Category=TakeRecorder)
 	FTakeRecorderPreInitialize TakeRecorderPreInitialize;
