TakeRecorder: Fix not incrementing the take number and creating duplicate asset names.
Fix not being able to open the same take that is already open.
Fix opening a sequence to record, then opening another sequence to view sometimes not opening in review.
Prevent recording while reviewing, fixing a possible crash.
Additional fixes for not recognizing CineAssembly assets.

diff --git a/Engine/Plugins/VirtualProduction/Takes/Source/TakeRecorder/Private/Recorder/TakeRecorderSubsystemImplementation.cpp b/Engine/Plugins/VirtualProduction/Takes/Source/TakeRecorder/Private/Recorder/TakeRecorderSubsystemImplementation.cpp
--- a/Engine/Plugins/VirtualProduction/Takes/Source/TakeRecorder/Private/Recorder/TakeRecorderSubsystemImplementation.cpp
+++ b/Engine/Plugins/VirtualProduction/Takes/Source/TakeRecorder/Private/Recorder/TakeRecorderSubsystemImplementation.cpp
@@ -134,28 +134,11 @@ void UTakeRecorderSubsystemImplementation::DeinitializeImplementation()
 
 void UTakeRecorderSubsystemImplementation::SetTargetSequence(const FTakeRecorderSequenceParameters& InData)
 {
-	// Don't initialize again if we already have the same data present.
-	if (TargetSequenceData.IsValid() && TargetSequenceData == InData)
-	{
-		return;
-	}
 	TargetSequenceData = InData;
 	
 	ITakeRecorderModule& TakeRecorderModule = FModuleManager::Get().LoadModuleChecked<ITakeRecorderModule>("TakeRecorder");
+	TakeRecorderModule.GetLastLevelSequenceProvider().Unbind();
 	TakeRecorderModule.GetLastLevelSequenceProvider().BindUObject(this, &UTakeRecorderSubsystemImplementation::SetLastLevelSequence);
-
-	bAutoApplyTakeNumber = true;
-	
-	CacheMetaData();
-	
-	if (TakeMetaData && !TakeMetaData->IsLocked())
-	{
-		const int32 NextTakeNumber = GetNextTakeNumber(TakeMetaData->GetSlate());
-		if (NextTakeNumber != TakeMetaData->GetTakeNumber())
-		{
-			TakeMetaData->SetTakeNumber(NextTakeNumber);
-		}
-	}
 	
 	// If a recording is currently underway, initialize to that now
 	if (UTakeRecorder* ActiveRecorder = UTakeRecorder::GetActiveRecorder())
@@ -213,8 +196,22 @@ void UTakeRecorderSubsystemImplementation::SetTargetSequence(const FTakeRecorder
 	else if (InData.SequenceToView)
 	{
 		SuppliedLevelSequence = InData.SequenceToView;
+		RecordIntoLevelSequence = nullptr; // We may have switched from recording to reviewing.
 	}
-
+	
+	bAutoApplyTakeNumber = true;
+	
+	CacheMetaData();
+	
+	if (TakeMetaData && !TakeMetaData->IsLocked())
+	{
+		const int32 NextTakeNumber = GetNextTakeNumber(TakeMetaData->GetSlate());
+		if (NextTakeNumber != TakeMetaData->GetTakeNumber())
+		{
+			TakeMetaData->SetTakeNumber(NextTakeNumber);
+		}
+	}
+	
 	/** Clear the dirty flag since the preset was just initialized. */
 	TransientPreset->GetOutermost()->SetDirtyFlag(false);
 
@@ -276,7 +273,7 @@ bool UTakeRecorderSubsystemImplementation::StartRecording(bool bOpenSequencer, b
 {
 	static bool bStartedRecording = false;
 
-	if (bStartedRecording)
+	if (bStartedRecording || IsReviewing() /** Shouldn't be starting a recording if we are reviewing. */)
 	{
 		return false;
 	}
@@ -595,6 +592,7 @@ TArray<FAssetData> UTakeRecorderSubsystemImplementation::GetSlates(FName InPacka
 
 	FARFilter Filter;
 	Filter.ClassPaths.Add(ULevelSequence::StaticClass()->GetClassPathName());
+	Filter.bRecursiveClasses = true;
 	Filter.TagsAndValues.Add(UTakeMetaData::AssetRegistryTag_Slate);
 	if (!InPackagePath.IsNone())
 	{
@@ -854,11 +852,7 @@ void UTakeRecorderSubsystemImplementation::CacheMetaData()
 			TransientTakeMetaData = UTakeMetaData::CreateFromDefaults(GetTransientPackage(), NAME_None);
 			TransientTakeMetaData->SetFlags(RF_Transactional | RF_Transient);
 
-			FString DefaultSlate = GetDefault<UTakeRecorderProjectSettings>()->Settings.DefaultSlate;
-			if (TransientTakeMetaData->GetSlate() != DefaultSlate)
-			{
-				TransientTakeMetaData->SetSlate(DefaultSlate, false);
-			}
+			UpdateTransientDefaultSlateName();
 
 			// Compute the correct starting take number
 			int32 NextTakeNumber = GetNextTakeNumber(TransientTakeMetaData->GetSlate());
@@ -898,6 +892,15 @@ void UTakeRecorderSubsystemImplementation::CacheMetaData()
 	check(TakeMetaData);
 }
 
+void UTakeRecorderSubsystemImplementation::UpdateTransientDefaultSlateName()
+{
+	const FString DefaultSlate = GetDefault<UTakeRecorderProjectSettings>()->Settings.DefaultSlate;
+	if (TransientTakeMetaData && TransientTakeMetaData->GetSlate() != DefaultSlate)
+	{
+		TransientTakeMetaData->SetSlate(DefaultSlate, false);
+	}
+}
+
 void UTakeRecorderSubsystemImplementation::SetFrameRateImpl(const FFrameRate& InFrameRate, bool bFromTimecode)
 {
 	if (TakeMetaData)
diff --git a/Engine/Plugins/VirtualProduction/Takes/Source/TakeRecorder/Private/Recorder/TakeRecorderSubsystemImplementation.h b/Engine/Plugins/VirtualProduction/Takes/Source/TakeRecorder/Private/Recorder/TakeRecorderSubsystemImplementation.h
--- a/Engine/Plugins/VirtualProduction/Takes/Source/TakeRecorder/Private/Recorder/TakeRecorderSubsystemImplementation.h
+++ b/Engine/Plugins/VirtualProduction/Takes/Source/TakeRecorder/Private/Recorder/TakeRecorderSubsystemImplementation.h
@@ -97,6 +97,9 @@ private:
 	/** Cache the current available metadata. */
 	void CacheMetaData();
 
+	/** Update our transient slate data with the current default slate name. */
+	void UpdateTransientDefaultSlateName();
+
 	/** Sets the frame rate. */
 	void SetFrameRateImpl(const FFrameRate& InFrameRate, bool bFromTimecode);
 
diff --git a/Engine/Plugins/VirtualProduction/Takes/Source/TakeRecorder/Private/Widgets/STakeRecorderPanel.cpp b/Engine/Plugins/VirtualProduction/Takes/Source/TakeRecorder/Private/Widgets/STakeRecorderPanel.cpp
--- a/Engine/Plugins/VirtualProduction/Takes/Source/TakeRecorder/Private/Widgets/STakeRecorderPanel.cpp
+++ b/Engine/Plugins/VirtualProduction/Takes/Source/TakeRecorder/Private/Widgets/STakeRecorderPanel.cpp
@@ -822,6 +822,7 @@ TSharedRef<SWidget> STakeRecorderPanel::OnOpenSequenceToRecordIntoMenu()
 
 		AssetPickerConfig.AssetShowWarningText = LOCTEXT("NoSequences_Warning", "No Level Sequences Found");
 		AssetPickerConfig.Filter.ClassPaths.Add(ULevelSequence::StaticClass()->GetClassPathName());
+		AssetPickerConfig.Filter.bRecursiveClasses = true;
 		AssetPickerConfig.OnAssetSelected = FOnAssetSelected::CreateSP(this, &STakeRecorderPanel::OnOpenSequenceToRecordInto);
 	}
 
