TakeRecorder: Change default of AnimationTrackName to {actor}_anim, solving an issue when loading into sequencer.

diff --git a/Engine/Plugins/VirtualProduction/Takes/Source/TakeTrackRecorders/Private/MovieSceneAnimationTrackRecorderSettings.cpp b/Engine/Plugins/VirtualProduction/Takes/Source/TakeTrackRecorders/Private/MovieSceneAnimationTrackRecorderSettings.cpp
--- a/Engine/Plugins/VirtualProduction/Takes/Source/TakeTrackRecorders/Private/MovieSceneAnimationTrackRecorderSettings.cpp
+++ b/Engine/Plugins/VirtualProduction/Takes/Source/TakeTrackRecorders/Private/MovieSceneAnimationTrackRecorderSettings.cpp
@@ -8,6 +8,18 @@
 
 #include UE_INLINE_GENERATED_CPP_BY_NAME(MovieSceneAnimationTrackRecorderSettings)
 
+UMovieSceneAnimationTrackRecorderEditorSettings::UMovieSceneAnimationTrackRecorderEditorSettings(const FObjectInitializer& ObjInit)
+	: Super(ObjInit)
+	, AnimationTrackName(NSLOCTEXT("UMovieSceneAnimationTrackRecorderSettings", "DefaultAnimationTrackName", "{actor}_anim"))
+	, AnimationAssetName(TEXT("{actor}_{takeName}"))
+	, AnimationSubDirectory(TEXT("Animation"))
+	, InterpMode()
+	, TangentMode()
+	, bRemoveRootAnimation(true)
+{
+	TimecodeBoneMethod.BoneMode = ETimecodeBoneMode::Root;
+}
+
 void UMovieSceneAnimationTrackRecorderEditorSettings::PostEditChangeProperty(FPropertyChangedEvent& PropertyChangedEvent)
 {
 	Super::PostEditChangeProperty(PropertyChangedEvent);
diff --git a/Engine/Plugins/VirtualProduction/Takes/Source/TakeTrackRecorders/Public/TrackRecorders/MovieSceneAnimationTrackRecorderSettings.h b/Engine/Plugins/VirtualProduction/Takes/Source/TakeTrackRecorders/Public/TrackRecorders/MovieSceneAnimationTrackRecorderSettings.h
--- a/Engine/Plugins/VirtualProduction/Takes/Source/TakeTrackRecorders/Public/TrackRecorders/MovieSceneAnimationTrackRecorderSettings.h
+++ b/Engine/Plugins/VirtualProduction/Takes/Source/TakeTrackRecorders/Public/TrackRecorders/MovieSceneAnimationTrackRecorderSettings.h
@@ -14,15 +14,7 @@ class TAKETRACKRECORDERS_API UMovieSceneAnimationTrackRecorderEditorSettings : p
 {
 	GENERATED_BODY()
 public:
-	UMovieSceneAnimationTrackRecorderEditorSettings(const FObjectInitializer& ObjInit)
-		: Super(ObjInit)
-	, AnimationTrackName(NSLOCTEXT("UMovieSceneAnimationTrackRecorderSettings", "DefaultAnimationTrackName", "RecordedAnimation"))
-	, AnimationAssetName(TEXT("{actor}_{takeName}"))
-	, AnimationSubDirectory(TEXT("Animation"))
-	, bRemoveRootAnimation(true)
-	{
-		TimecodeBoneMethod.BoneMode = ETimecodeBoneMode::Root;
-	}
+	UMovieSceneAnimationTrackRecorderEditorSettings(const FObjectInitializer& ObjInit);
 
 	virtual void PostEditChangeProperty(FPropertyChangedEvent& PropertyChangedEvent) override;
 
