TakeRecorder: Transact when removing a source from the subsystem, fixing MU usage.

diff --git a/Engine/Plugins/VirtualProduction/Takes/Source/TakeRecorder/Private/Recorder/TakeRecorderSubsystemImplementation.cpp b/Engine/Plugins/VirtualProduction/Takes/Source/TakeRecorder/Private/Recorder/TakeRecorderSubsystemImplementation.cpp
--- a/Engine/Plugins/VirtualProduction/Takes/Source/TakeRecorder/Private/Recorder/TakeRecorderSubsystemImplementation.cpp
+++ b/Engine/Plugins/VirtualProduction/Takes/Source/TakeRecorder/Private/Recorder/TakeRecorderSubsystemImplementation.cpp
@@ -445,7 +445,7 @@ UTakeRecorderSource* UTakeRecorderSubsystemImplementation::AddSource(const TSubc
 
 	if (*InSourceClass && Sources)
 	{
-		FScopedTransaction Transaction(FText::Format(LOCTEXT("AddNewSource", "Add New {0} Source"), InSourceClass->GetDisplayNameText()));
+		FScopedTransaction Transaction(FText::Format(LOCTEXT("AddNewSource_Transaction", "Add New {0} Source"), InSourceClass->GetDisplayNameText()));
 		Sources->Modify();
 
 		return Sources->AddSource(InSourceClass);
@@ -459,6 +459,8 @@ void UTakeRecorderSubsystemImplementation::RemoveSource(UTakeRecorderSource* InS
 	UTakeRecorderSources* Sources = GetSources();
 	if (Sources && InSource)
 	{
+		FScopedTransaction Transaction(FText::Format(LOCTEXT("RemoveSource_Transaction", "Remove {0} Source"), InSource->GetDisplayText()));
+		Sources->Modify();
 		Sources->RemoveSource(InSource);
 	}
 }
@@ -467,6 +469,8 @@ void UTakeRecorderSubsystemImplementation::ClearSources()
 {
 	if (UTakeRecorderSources* Sources = GetSources())
 	{
+		FScopedTransaction Transaction(LOCTEXT("ClearSources_Transaction", "Clear Sources"));
+		Sources->Modify();
 		for (UTakeRecorderSource* Source : Sources->GetSourcesCopy())
 		{
 			Sources->RemoveSource(Source);
