TakeRecorder: Fix crash when recording with a deleted (none) actor.

diff --git a/Engine/Plugins/VirtualProduction/Takes/Source/TakeRecorderSources/Private/TakeRecorderActorSource.cpp b/Engine/Plugins/VirtualProduction/Takes/Source/TakeRecorderSources/Private/TakeRecorderActorSource.cpp
--- a/Engine/Plugins/VirtualProduction/Takes/Source/TakeRecorderSources/Private/TakeRecorderActorSource.cpp
+++ b/Engine/Plugins/VirtualProduction/Takes/Source/TakeRecorderSources/Private/TakeRecorderActorSource.cpp
@@ -374,6 +374,11 @@ void UTakeRecorderActorSource::StartRecording(const FTimecode& InSectionStartTim
 	// their starts to the most up to date time sample.
 	for (UMovieSceneTrackRecorder* Recorder : TrackRecorders)
 	{
+		if (!Recorder)
+		{
+			UE_LOG(LogTakesCore, Error, TEXT("Start Recording: Actor source %s has an invalid recorder."), *Target.GetAssetName());
+			continue;
+		}
 		// The Frame Number is the sequence time we should record into, but we want to get
 		// a Timecode stamp that reflects the real time it was recorded... This kind of conflicts
 		// with 
@@ -623,7 +628,7 @@ void UTakeRecorderActorSource::TickRecording(const FQualifiedFrameTime& CurrentS
 			for (UMovieSceneTrackRecorder* TrackRecorder : TrackRecorders)
 			{
 				// If the track recorder existed before we added this component then it has already had StartRecording called on it.
-				if (PreviousTrackRecorders.Contains(TrackRecorder))
+				if (!TrackRecorder || PreviousTrackRecorders.Contains(TrackRecorder))
 				{
 					continue;
 				}
@@ -647,7 +652,10 @@ void UTakeRecorderActorSource::TickRecording(const FQualifiedFrameTime& CurrentS
 	// Now that we've initialized any new components we can tick all of our recordings to get the last frame's data.
 	for (UMovieSceneTrackRecorder* Recorder : TrackRecorders)
 	{
-		Recorder->RecordSample(CurrentSequenceTime);
+		if (Recorder)
+		{
+			Recorder->RecordSample(CurrentSequenceTime);
+		}
 	}
 
 	// Update our cached list
@@ -658,6 +666,11 @@ void UTakeRecorderActorSource::StopRecording(ULevelSequence* InSequence)
 {
 	for (UMovieSceneTrackRecorder* TrackRecorder : TrackRecorders)
 	{
+		if (!TrackRecorder)
+		{
+			UE_LOG(LogTakesCore, Error, TEXT("Stop Recording: Actor source %s has an invalid recorder."), *Target.GetAssetName());
+			continue;
+		}
 		TrackRecorder->StopRecording();
 	}
 	
@@ -762,6 +775,12 @@ TArray<UTakeRecorderSource*> UTakeRecorderActorSource::PostRecording(ULevelSeque
 	int32 SectionRecorderIndex = 0;
 	for (UMovieSceneTrackRecorder* SectionRecorder : TrackRecorders)
 	{
+		if (!SectionRecorder)
+		{
+			UE_LOG(LogTakesCore, Error, TEXT("Post Recording: Actor source %s has an invalid section recorder."), *Target.GetAssetName());
+			continue;
+		}
+		
 		// Increment before entering the progress frame so we get "1/1" instead of "0/1"
 		SectionRecorderIndex++;
 
@@ -871,6 +890,11 @@ void UTakeRecorderActorSource::PostProcessTrackRecorders(ULevelSequence* InSeque
 
 	for (UMovieSceneTrackRecorder* TrackRecorder : TrackRecorders)
 	{
+		if (!TrackRecorder)
+		{
+			UE_LOG(LogTakesCore, Error, TEXT("Post Process Track Recorders: Actor source %s has an invalid recorder."), *Target.GetAssetName());
+			continue;
+		}
 		AActor* SourceActor = Cast<AActor>(TrackRecorder->GetSourceObject());
 		AActor* SourceEditorActor = SourceActor ? EditorUtilities::GetEditorWorldCounterpartActor(SourceActor) : nullptr;
 		AActor* TargetActor = Target.Get();
@@ -900,6 +924,12 @@ void UTakeRecorderActorSource::PostProcessTrackRecorders(ULevelSequence* InSeque
 	// Remove root motion on all other animation track recorders
 	for (UMovieSceneTrackRecorder* TrackRecorder : TrackRecorders)
 	{
+		if (!TrackRecorder)
+		{
+			UE_LOG(LogTakesCore, Error, TEXT("Post Process Track Recorders: Actor source %s has an invalid recorder."), *Target.GetAssetName());
+			continue;
+		}
+		
 		if (TrackRecorder->IsA<UMovieSceneAnimationTrackRecorder>())
 		{
 			UMovieSceneAnimationTrackRecorder* AnimationTrackRecorder = Cast<UMovieSceneAnimationTrackRecorder>(TrackRecorder);
diff --git a/Engine/Plugins/VirtualProduction/Takes/Source/TakeRecorderSources/Private/TakeRecorderSourcesUtils.h b/Engine/Plugins/VirtualProduction/Takes/Source/TakeRecorderSources/Private/TakeRecorderSourcesUtils.h
--- a/Engine/Plugins/VirtualProduction/Takes/Source/TakeRecorderSources/Private/TakeRecorderSourcesUtils.h
+++ b/Engine/Plugins/VirtualProduction/Takes/Source/TakeRecorderSources/Private/TakeRecorderSourcesUtils.h
@@ -142,9 +142,9 @@ namespace TakeRecorderSourcesUtils
 				{
 					for (UMovieSceneTrackRecorder* TrackRecorder : ActorSource->TrackRecorders)
 					{
-						if (TrackRecorder->IsA<UMovieSceneAnimationTrackRecorder>())
+						if (TrackRecorder && TrackRecorder->IsA<UMovieSceneAnimationTrackRecorder>())
 						{
-							return Cast<UMovieSceneAnimationTrackRecorder>(TrackRecorder)->GetInitialRootTransform();
+							return CastChecked<UMovieSceneAnimationTrackRecorder>(TrackRecorder)->GetInitialRootTransform();
 						}
 					}
 
