Sequencer: tear down Take Recorder's Sequencer instance before we delete and GC the sequence

The Sequencer instance doesn't really expect its currently edited sequence to be GC'ed from underneath it.

diff --git a/Engine/Plugins/VirtualProduction/Takes/Source/TakeRecorder/Private/Recorder/TakeRecorder.cpp b/Engine/Plugins/VirtualProduction/Takes/Source/TakeRecorder/Private/Recorder/TakeRecorder.cpp
--- a/Engine/Plugins/VirtualProduction/Takes/Source/TakeRecorder/Private/Recorder/TakeRecorder.cpp
+++ b/Engine/Plugins/VirtualProduction/Takes/Source/TakeRecorder/Private/Recorder/TakeRecorder.cpp
@@ -1248,6 +1248,9 @@ void UTakeRecorder::StopInternal(const bool bCancelled)
 		}
 	}
 
+	// Tear down the Sequencer instance before we delete the SequenceAsset.
+	Sequencer.Reset();
+
 	// Delete the asset after OnTakeRecorderFinished and OnTakeRecorderCancelled because the ScopedSequencerPanel 
 	// will still have the current sequence before it is returned to the Pending Take.
 	if (Parameters.TakeRecorderMode == ETakeRecorderMode::RecordNewSequence && !bRecordingFinished)
