Updated headers using UnrealCodeFixup to make sure dllstorage is on methods/staticvars instead of types. Part 4/n

diff --git a/Engine/Plugins/VirtualProduction/Takes/Source/TakeRecorder/Public/Recorder/TakeRecorderSubsystem.h b/Engine/Plugins/VirtualProduction/Takes/Source/TakeRecorder/Public/Recorder/TakeRecorderSubsystem.h
--- a/Engine/Plugins/VirtualProduction/Takes/Source/TakeRecorder/Public/Recorder/TakeRecorderSubsystem.h
+++ b/Engine/Plugins/VirtualProduction/Takes/Source/TakeRecorder/Public/Recorder/TakeRecorderSubsystem.h
@@ -9,6 +9,8 @@
 
 #include "TakeRecorderSubsystem.generated.h"
 
+#define UE_API TAKERECORDER_API
+
 struct FNamingTokensEvaluationData;
 class ULevelSequence;
 class UTakeRecorderNamingTokensData;
@@ -37,15 +39,15 @@ DECLARE_DYNAMIC_MULTICAST_DELEGATE_OneParam(FOnTakeRecorderSourceRemovedDynamic,
 * SetTargetSequence() needs to be called at least once prior to use. This will perform additional initialization
 * and register the subsystem as tickable currently required to cache level metadata.
 */
-UCLASS()
-class TAKERECORDER_API UTakeRecorderSubsystem final : public UEngineSubsystem, public ITakeRecorderSubsystemInterface
+UCLASS(MinimalAPI)
+class UTakeRecorderSubsystem final : public UEngineSubsystem, public ITakeRecorderSubsystemInterface
 {
 	GENERATED_BODY()
 
 public:
 	
-	virtual void Initialize(FSubsystemCollectionBase& Collection) override;
-	virtual void Deinitialize() override;
+	UE_API virtual void Initialize(FSubsystemCollectionBase& Collection) override;
+	UE_API virtual void Deinitialize() override;
 
 	// ~Begin ITakeRecorderSubsystem interface
 	
@@ -56,21 +58,21 @@ public:
 	 * @param InData The sequence parameters, which are mutually exclusive. They can all be null.
 	 */
 	UFUNCTION(BlueprintCallable, Category = TakeRecorder, meta=(AutoCreateRefTerm="InData"))
-	virtual void SetTargetSequence(const FTakeRecorderSequenceParameters& InData = FTakeRecorderSequenceParameters()) override;
+	UE_API virtual void SetTargetSequence(const FTakeRecorderSequenceParameters& InData = FTakeRecorderSequenceParameters()) override;
 	
 	/** Set the record into level sequence. */
-	virtual void SetRecordIntoLevelSequence(ULevelSequence* LevelSequence) override;
+	UE_API virtual void SetRecordIntoLevelSequence(ULevelSequence* LevelSequence) override;
 
 	/** Can we review the last recording? */
 	UFUNCTION(BlueprintCallable, Category = TakeRecorder)
-	virtual bool CanReviewLastRecording() const override;
+	UE_API virtual bool CanReviewLastRecording() const override;
 	
 	/**
 	 * Supply the last recording if it exists.
 	 * @return true if it can be reviewed, false if it there isn't a recording to review.
 	 */
 	UFUNCTION(BlueprintCallable, Category = TakeRecorder)
-	virtual bool ReviewLastRecording() override;
+	UE_API virtual bool ReviewLastRecording() override;
 
 	/**
 	 * Begin a new recording.
@@ -79,59 +81,59 @@ public:
 	 * @return true if the recording was started successfully.
 	 */
 	UFUNCTION(BlueprintCallable, Category = TakeRecorder)
-	virtual bool StartRecording(bool bOpenSequencer = true, bool bShowErrorMessage = true) override;
+	UE_API virtual bool StartRecording(bool bOpenSequencer = true, bool bShowErrorMessage = true) override;
 
 	/**
 	 * Stop an existing recording.
 	 */
 	UFUNCTION(BlueprintCallable, Category = TakeRecorder)
-	virtual void StopRecording() override;
+	UE_API virtual void StopRecording() override;
 
 	/**
 	 * Cancel an in-progress recording.
 	 */
 	UFUNCTION(BlueprintCallable, Category = TakeRecorder)
-	virtual void CancelRecording(bool bShowConfirmMessage = true) override;
+	UE_API virtual void CancelRecording(bool bShowConfirmMessage = true) override;
 
 	/**
 	 * Reset to the pending take.
 	 */
 	UFUNCTION(BlueprintCallable, Category = TakeRecorder)
-	virtual void ResetToPendingTake() override;
+	UE_API virtual void ResetToPendingTake() override;
 
 	/**
 	 * Clear the pending take.
 	 */
 	UFUNCTION(BlueprintCallable, Category = TakeRecorder)
-	virtual void ClearPendingTake() override;
+	UE_API virtual void ClearPendingTake() override;
 
 	/** Retrieve the pending take. This may be null. */
 	UFUNCTION(BlueprintCallable, Category = TakeRecorder)
-	virtual UTakePreset* GetPendingTake() const override;
+	UE_API virtual UTakePreset* GetPendingTake() const override;
 
 	/** Revert any changes restoring the preset origin. */
 	UFUNCTION(BlueprintCallable, Category = TakeRecorder)
-	virtual void RevertChanges() override;
+	UE_API virtual void RevertChanges() override;
 	
 	// ~Begin Sources
 	/** Add a source by a source class. */
 	UFUNCTION(BlueprintCallable, Category = TakeRecorder, meta = (DeterminesOutputType = "InSourceClass"))
-	virtual UTakeRecorderSource* AddSource(const TSubclassOf<UTakeRecorderSource> InSourceClass) override;
+	UE_API virtual UTakeRecorderSource* AddSource(const TSubclassOf<UTakeRecorderSource> InSourceClass) override;
 	
 	/** Remove a given source. */
 	UFUNCTION(BlueprintCallable, Category = TakeRecorder)
-	virtual void RemoveSource(UTakeRecorderSource* InSource) override;
+	UE_API virtual void RemoveSource(UTakeRecorderSource* InSource) override;
 
 	/** Remove all sources from the current sequence. */
 	UFUNCTION(BlueprintCallable, Category = TakeRecorder)
-	virtual void ClearSources() override;
+	UE_API virtual void ClearSources() override;
 
 	/** Retrieve the sources. */
 	UFUNCTION(BlueprintCallable, Category = TakeRecorder)
-	virtual UTakeRecorderSources* GetSources() const override;
+	UE_API virtual UTakeRecorderSources* GetSources() const override;
 	
 	/** Retrieve all sources for the current sequence. */
-	virtual TArrayView<UTakeRecorderSource* const> GetAllSources() const override;
+	UE_API virtual TArrayView<UTakeRecorderSource* const> GetAllSources() const override;
 
 	/**
 	 * Retrieves a copy of the list of sources that are being recorded. This is intended for Blueprint usages which cannot
@@ -139,11 +141,11 @@ public:
 	 * DO NOT MODIFY THIS ARRAY, modifications will be lost.
 	 */
 	UFUNCTION(BlueprintPure, DisplayName = "Get Sources (Copy)", Category = TakeRecorder)
-	virtual TArray<UTakeRecorderSource*> GetAllSourcesCopy() const override;
+	UE_API virtual TArray<UTakeRecorderSource*> GetAllSourcesCopy() const override;
 
 	/** Retrieve the first source of the given class. */
 	UFUNCTION(BlueprintCallable, Category = TakeRecorder, meta = (DeterminesOutputType = "InSourceClass"))
-	virtual UTakeRecorderSource* GetSourceByClass(const TSubclassOf<UTakeRecorderSource> InSourceClass) const override;
+	UE_API virtual UTakeRecorderSource* GetSourceByClass(const TSubclassOf<UTakeRecorderSource> InSourceClass) const override;
 
 	/*
 	 * Add an actor as a source.
@@ -153,22 +155,22 @@ public:
 	 * @param bShowProgress Enable/disable the dialog box showing progress for the potentially slow parts of finalizing the take
 	 */
 	UFUNCTION(BlueprintCallable, Category = TakeRecorder)
-	virtual void AddSourceForActor(AActor* InActor, bool bReduceKeys = true, bool bShowProgress = true) override;
+	UE_API virtual void AddSourceForActor(AActor* InActor, bool bReduceKeys = true, bool bShowProgress = true) override;
 
 	/**
 	 * Remove an actor from available sources.
 	 */
 	UFUNCTION(BlueprintCallable, Category = TakeRecorder)
-	virtual void RemoveActorFromSources(AActor* InActor) override;
+	UE_API virtual void RemoveActorFromSources(AActor* InActor) override;
 
 	/** Retrieve the actor from a source, if applicable. */
 	UFUNCTION(BlueprintCallable, Category = TakeRecorder)
-	virtual AActor* GetSourceActor(UTakeRecorderSource* InSource) const override;
+	UE_API virtual AActor* GetSourceActor(UTakeRecorderSource* InSource) const override;
 	// ~End Sources
 	
 	/** Retrieve the current take recorder state. */
 	UFUNCTION(BlueprintCallable, Category = TakeRecorder)
-	virtual ETakeRecorderState GetState() const override;
+	UE_API virtual ETakeRecorderState GetState() const override;
 
 	/**
 	 * Directly set the take number.
@@ -176,22 +178,22 @@ public:
 	 * @param bEmitChanged Whether to broadcast events signaling the take number has changed.
 	 */
 	UFUNCTION(BlueprintCallable, Category = TakeRecorder)
-	virtual void SetTakeNumber(int32 InNewTakeNumber, bool bEmitChanged = true) override;
+	UE_API virtual void SetTakeNumber(int32 InNewTakeNumber, bool bEmitChanged = true) override;
 
 	/** Compute the next take number given a slate. */
 	UFUNCTION(BlueprintCallable, Category = TakeRecorder)
-	virtual int32 GetNextTakeNumber(const FString& InSlate) const override;
+	UE_API virtual int32 GetNextTakeNumber(const FString& InSlate) const override;
 
 	/** Find both the current maximum take value and the total number of takes for a given slate. */
 	UFUNCTION(BlueprintCallable, Category = TakeRecorder)
-	virtual void GetNumberOfTakes(const FString& InSlate, int32& OutMaxTake, int32& OutNumTakes) const override;
+	UE_API virtual void GetNumberOfTakes(const FString& InSlate, int32& OutMaxTake, int32& OutNumTakes) const override;
 
 	/**
 	 * Retrieve all slates.
 	 * @param InPackagePath [Optional] Scope the search to a specific folder, recursively.
 	 */
 	UFUNCTION(BlueprintCallable, Category = TakeRecorder)
-	virtual TArray<FAssetData> GetSlates(FName InPackagePath = NAME_None) const override;
+	UE_API virtual TArray<FAssetData> GetSlates(FName InPackagePath = NAME_None) const override;
 	
 	/**
 	 * Directly set the slate name.
@@ -199,40 +201,40 @@ public:
 	 * @param bEmitChanged Whether to broadcast events signaling the slate name has changed.
 	 */
 	UFUNCTION(BlueprintCallable, Category = TakeRecorder)
-	virtual void SetSlateName(const FString& InSlateName, bool bEmitChanged = true) override;
+	UE_API virtual void SetSlateName(const FString& InSlateName, bool bEmitChanged = true) override;
 
 	/** Mark the current frame. */
 	UFUNCTION(BlueprintCallable, Category = TakeRecorder)
-	virtual bool MarkFrame() override;
+	UE_API virtual bool MarkFrame() override;
 
 	/**
 	 * Access the frame rate for this take
 	 */
 	UFUNCTION(BlueprintCallable, Category = TakeRecorder)
-	virtual FFrameRate GetFrameRate() const override;
+	UE_API virtual FFrameRate GetFrameRate() const override;
 	
 	/**
 	 * Set the frame rate for this take
 	 */
 	UFUNCTION(BlueprintCallable, Category = TakeRecorder)
-	virtual void SetFrameRate(FFrameRate InFrameRate) override;
+	UE_API virtual void SetFrameRate(FFrameRate InFrameRate) override;
 	
 	/**
 	 * Set if the frame rate is set from the Timecode frame rate
 	 */
 	UFUNCTION(BlueprintCallable, Category = TakeRecorder)
-	virtual void SetFrameRateFromTimecode() override;
+	UE_API virtual void SetFrameRateFromTimecode() override;
 
 	/** Import a preset to the transient preset. */
-	virtual void ImportPreset(const FAssetData& InPreset) override;
+	UE_API virtual void ImportPreset(const FAssetData& InPreset) override;
 
 	/** If Take Recorder is currently reviewing. */
 	UFUNCTION(BlueprintCallable, Category = TakeRecorder)
-	virtual bool IsReviewing() const override;
+	UE_API virtual bool IsReviewing() const override;
 
 	/** If Take Recorder is currently recording. */
 	UFUNCTION(BlueprintCallable, Category = TakeRecorder)
-	virtual bool IsRecording() const override;
+	UE_API virtual bool IsRecording() const override;
 
 	/**
 	 * Retrieve the current sequence's countdown.
@@ -240,93 +242,93 @@ public:
 	 * @return True if we are in a countdown sequence, false otherwise.
 	 */
 	UFUNCTION(BlueprintCallable, Category = TakeRecorder)
-	virtual bool TryGetSequenceCountdown(float& OutValue) const override;
+	UE_API virtual bool TryGetSequenceCountdown(float& OutValue) const override;
 
 	/**
 	 * Sets the current sequence's countdown.
 	 * @param InSeconds Time in seconds.
 	 */
 	UFUNCTION(BlueprintCallable, Category = TakeRecorder)
-	virtual void SetSequenceCountdown(float InSeconds) override;
+	UE_API virtual void SetSequenceCountdown(float InSeconds) override;
 
 	/** Retrieve additional settings objects from a source. */
 	UFUNCTION(BlueprintCallable, Category = TakeRecorder)
-	virtual TArray<UObject*> GetSourceRecordSettings(UTakeRecorderSource* InSource) const override;
+	UE_API virtual TArray<UObject*> GetSourceRecordSettings(UTakeRecorderSource* InSource) const override;
 
 	/** Retrieve the global take recorder settings. */
 	UFUNCTION(BlueprintCallable, Category = TakeRecorder)
-	virtual FTakeRecorderParameters GetGlobalRecordSettings() const override;
+	UE_API virtual FTakeRecorderParameters GetGlobalRecordSettings() const override;
 
 	/** Set the global take recorder settings. */
 	UFUNCTION(BlueprintCallable, Category = TakeRecorder)
-	virtual void SetGlobalRecordSettings(const FTakeRecorderParameters& InParameters) override;
+	UE_API virtual void SetGlobalRecordSettings(const FTakeRecorderParameters& InParameters) override;
 	
 	/** Retrieve the current meta data. */
 	UFUNCTION(BlueprintCallable, Category = TakeRecorder)
-	virtual UTakeMetaData* GetTakeMetaData() const override;
+	UE_API virtual UTakeMetaData* GetTakeMetaData() const override;
 
 	/** Return the level sequence we are using. */
 	UFUNCTION(BlueprintCallable, Category = TakeRecorder)
-	virtual ULevelSequence* GetLevelSequence() const override;
+	UE_API virtual ULevelSequence* GetLevelSequence() const override;
 
 	/** Retrieve the Supplied Level Sequence if it exists. */
-	virtual ULevelSequence* GetSuppliedLevelSequence() const override;
+	UE_API virtual ULevelSequence* GetSuppliedLevelSequence() const override;
 
 	/** Retrieve the Recording Level Sequence if it exists. */
-	virtual ULevelSequence* GetRecordingLevelSequence() const override;
+	UE_API virtual ULevelSequence* GetRecordingLevelSequence() const override;
 
 	/** Retrieve the Record Into Level Sequence if it exists. */
-	virtual ULevelSequence* GetRecordIntoLevelSequence() const override;
+	UE_API virtual ULevelSequence* GetRecordIntoLevelSequence() const override;
 
 	/** Retrieve the Last Recorded Level Sequence if it exists. */
 	UFUNCTION(BlueprintCallable, Category = TakeRecorder)
-	virtual ULevelSequence* GetLastRecordedLevelSequence() const override;
+	UE_API virtual ULevelSequence* GetLastRecordedLevelSequence() const override;
 
 	/** Retrieve the transient preset the subsystem is using. */
-	virtual UTakePreset* GetTransientPreset() const override;
+	UE_API virtual UTakePreset* GetTransientPreset() const override;
 	
 	/** The current take recorder mode. */
 	UFUNCTION(BlueprintCallable, Category = TakeRecorder)
-	virtual ETakeRecorderMode GetTakeRecorderMode() const override;
+	UE_API virtual ETakeRecorderMode GetTakeRecorderMode() const override;
 
 	/** Retrieve relevant Naming Tokens data for Take Recorder. */
-	virtual UTakeRecorderNamingTokensData* GetNamingTokensData() const override;
+	UE_API virtual UTakeRecorderNamingTokensData* GetNamingTokensData() const override;
 
 	/** @return Whether there are any pending changes, which ClearPendingTake would discard. */
-	virtual bool HasPendingChanges() const override;
+	UE_API virtual bool HasPendingChanges() const override;
 	
 	/**
 	 * Retrieve a multi-cast delegate that is triggered when a recording pre-initializes
 	 */
-	FOnTakeRecordingInitialized& GetOnRecordingPreInitializedEvent();
+	UE_API FOnTakeRecordingInitialized& GetOnRecordingPreInitializedEvent();
 	/**
 	 * Retrieve a multi-cast delegate that is triggered when a recording initializes
 	 */
-	FOnTakeRecordingInitialized& GetOnRecordingInitializedEvent();
+	UE_API FOnTakeRecordingInitialized& GetOnRecordingInitializedEvent();
 	/**
 	 * Retrieve a multi-cast delegate that is triggered when a recording starts
 	 */
-	FOnTakeRecordingStarted& GetOnRecordingStartedEvent();
+	UE_API FOnTakeRecordingStarted& GetOnRecordingStartedEvent();
 	/**
 	 * Retrieve a multi-cast delegate that is triggered when a recording is stopped
 	 */
-	FOnTakeRecordingStopped& GetOnRecordingStoppedEvent();
+	UE_API FOnTakeRecordingStopped& GetOnRecordingStoppedEvent();
 	/**
 	 * Retrieve a multi-cast delegate that is triggered when a recording finishes
 	 */
-	FOnTakeRecordingFinished& GetOnRecordingFinishedEvent();
+	UE_API FOnTakeRecordingFinished& GetOnRecordingFinishedEvent();
 	/**
 	 * Retrieve a multi-cast delegate that is triggered when a recording is cancelled.
 	 */
-	FOnTakeRecordingCancelled& GetOnRecordingCancelledEvent();
+	UE_API FOnTakeRecordingCancelled& GetOnRecordingCancelledEvent();
 	/**
 	 * Retrieve a multi-cast delegate that is triggered when a source is added.
 	 */
-	UTakeRecorderSources::FOnSourceAdded& GetOnRecordingSourceAddedEvent();
+	UE_API UTakeRecorderSources::FOnSourceAdded& GetOnRecordingSourceAddedEvent();
 	/**
 	 * Retrieve a multi-cast delegate that is triggered when a source is removed.
 	 */
-	UTakeRecorderSources::FOnSourceRemoved& GetOnRecordingSourceRemovedEvent();
+	UE_API UTakeRecorderSources::FOnSourceRemoved& GetOnRecordingSourceRemovedEvent();
 
 	// ~End ITakeRecorderSubsystem interface
 	
@@ -406,3 +408,5 @@ public:
 	UPROPERTY(BlueprintAssignable, Category=TakeRecorder)
 	FOnTakeRecorderSourceRemovedDynamic TakeRecorderSourceRemoved;
 };
+
+#undef UE_API
diff --git a/Engine/Plugins/VirtualProduction/Takes/Source/TakeRecorder/Public/SerializedRecorder.h b/Engine/Plugins/VirtualProduction/Takes/Source/TakeRecorder/Public/SerializedRecorder.h
--- a/Engine/Plugins/VirtualProduction/Takes/Source/TakeRecorder/Public/SerializedRecorder.h
+++ b/Engine/Plugins/VirtualProduction/Takes/Source/TakeRecorder/Public/SerializedRecorder.h
@@ -5,31 +5,33 @@
 #include "ISerializedRecorder.h"
 #include "LevelSequence.h"
 
+#define UE_API TAKERECORDER_API
+
 class UWorld;
 class AActor;
 struct FActorFileHeader;
 struct FActorProperty;
 //Implementation of ISerializedRecorder that's defined in the SerializedRecorderInterface Module
 
-class TAKERECORDER_API FSerializedRecorder:  public ISerializedRecorder
+class FSerializedRecorder:  public ISerializedRecorder
 {
 public:
 	FSerializedRecorder() : bLoadSequenceFile(false) {}
 	~FSerializedRecorder() {}
 
-	virtual bool LoadRecordedSequencerFile(UMovieSceneSequence* InMovieSceneSequence, UWorld *PlaybackContext, const FString& InFileName, TFunction<void()> InCompletionCallback) override;
+	UE_API virtual bool LoadRecordedSequencerFile(UMovieSceneSequence* InMovieSceneSequence, UWorld *PlaybackContext, const FString& InFileName, TFunction<void()> InCompletionCallback) override;
 
 	//this is only valid during LoadSequenceFile or LoadActorFile
-	AActor GetActorFromGuid(const FGuid& InGuid);
+	UE_API AActor GetActorFromGuid(const FGuid& InGuid);
 
 private:
-	bool LoadSequenceFile(UMovieSceneSequence* InMovieSceneSequence, UWorld *PlaybackContext, const FString& InFileName, TFunction<void()> InCompletionCallback);
-	bool LoadActorFile(UMovieSceneSequence* InMovieSceneSequence, UWorld *PlaybackContext, const FString& InFileName, TFunction<void()> InCompletionCallback);
-	bool LoadPropertyFile(UMovieSceneSequence* InMovieSceneSequence, UWorld *PlaybackContext, const FString& InFileName, TFunction<void()> InCompletionCallback);
-	bool LoadSubSequenceFile(UMovieSceneSequence* InMovieSceneSequence, UWorld *PlaybackContext, const FString& InFileName, TFunction<void()> InCompletionCallback);
+	UE_API bool LoadSequenceFile(UMovieSceneSequence* InMovieSceneSequence, UWorld *PlaybackContext, const FString& InFileName, TFunction<void()> InCompletionCallback);
+	UE_API bool LoadActorFile(UMovieSceneSequence* InMovieSceneSequence, UWorld *PlaybackContext, const FString& InFileName, TFunction<void()> InCompletionCallback);
+	UE_API bool LoadPropertyFile(UMovieSceneSequence* InMovieSceneSequence, UWorld *PlaybackContext, const FString& InFileName, TFunction<void()> InCompletionCallback);
+	UE_API bool LoadSubSequenceFile(UMovieSceneSequence* InMovieSceneSequence, UWorld *PlaybackContext, const FString& InFileName, TFunction<void()> InCompletionCallback);
 
-	AActor* SetActorPossesableOrSpawnable(UMovieSceneSequence* InMovieSceneSequence, UWorld *PlaybackContext, const FActorFileHeader& ActorHeader);
-	void SetComponentPossessable(UMovieSceneSequence* InMovieSceneSequence, UWorld *PlaybackContext, AActor* Actor, const FActorFileHeader& ActorHeader, const FActorProperty& ActorProperty);
+	UE_API AActor* SetActorPossesableOrSpawnable(UMovieSceneSequence* InMovieSceneSequence, UWorld *PlaybackContext, const FActorFileHeader& ActorHeader);
+	UE_API void SetComponentPossessable(UMovieSceneSequence* InMovieSceneSequence, UWorld *PlaybackContext, AActor* Actor, const FActorFileHeader& ActorHeader, const FActorProperty& ActorProperty);
 
 
 private:
@@ -38,3 +40,4 @@ private:
 	
 };
 
+#undef UE_API
diff --git a/Engine/Plugins/VirtualProduction/Takes/Source/TakeRecorder/Public/TakeRecorderOverlayWidget.h b/Engine/Plugins/VirtualProduction/Takes/Source/TakeRecorder/Public/TakeRecorderOverlayWidget.h
--- a/Engine/Plugins/VirtualProduction/Takes/Source/TakeRecorder/Public/TakeRecorderOverlayWidget.h
+++ b/Engine/Plugins/VirtualProduction/Takes/Source/TakeRecorder/Public/TakeRecorderOverlayWidget.h
@@ -5,16 +5,18 @@
 #include "Blueprint/UserWidget.h"
 #include "TakeRecorderOverlayWidget.generated.h"
 
+#define UE_API TAKERECORDER_API
+
 class UTakeRecorder;
 
-UCLASS(Blueprintable, BlueprintType)
-class TAKERECORDER_API UTakeRecorderOverlayWidget : public UUserWidget
+UCLASS(MinimalAPI, Blueprintable, BlueprintType)
+class UTakeRecorderOverlayWidget : public UUserWidget
 {
 public:
 
 	GENERATED_BODY()
 
-	UTakeRecorderOverlayWidget(const FObjectInitializer& ObjectInitializer);
+	UE_API UTakeRecorderOverlayWidget(const FObjectInitializer& ObjectInitializer);
 
 	/**
 	 * Set the recorder that this overlay is reflecting
@@ -29,4 +31,6 @@ protected:
 	/** The recorder that this overlay is reflecting */
 	UPROPERTY(BlueprintReadOnly, Category="Take Recorder")
 	TObjectPtr<UTakeRecorder> Recorder;
-};
\ No newline at end of file
+};
+
+#undef UE_API
diff --git a/Engine/Plugins/VirtualProduction/Takes/Source/TakeRecorderSources/Private/TakeRecorderMicrophoneAudioManager.h b/Engine/Plugins/VirtualProduction/Takes/Source/TakeRecorderSources/Private/TakeRecorderMicrophoneAudioManager.h
--- a/Engine/Plugins/VirtualProduction/Takes/Source/TakeRecorderSources/Private/TakeRecorderMicrophoneAudioManager.h
+++ b/Engine/Plugins/VirtualProduction/Takes/Source/TakeRecorderSources/Private/TakeRecorderMicrophoneAudioManager.h
@@ -6,6 +6,8 @@
 
 #include "TakeRecorderMicrophoneAudioManager.generated.h"
 
+#define UE_API TAKERECORDERSOURCES_API
+
 class IAudioCaptureEditor;
 class USoundWave;
 struct FTakeRecorderAudioSourceSettings;
@@ -17,23 +19,23 @@ DECLARE_MULTICAST_DELEGATE_OneParam(FOnNotifySourcesOfDeviceChange, int);
 *   conjunction with FAudioInputDevicePropertyCustomization. It also manages the IAudioCaptureEditor 
 *   object which handles the low level audio device recording.
 */
-UCLASS(config=EditorPerProjectUserSettings, PerObjectConfig, DisplayName = "Audio Input Device")
-class TAKERECORDERSOURCES_API UTakeRecorderMicrophoneAudioManager : public UTakeRecorderAudioInputSettings
+UCLASS(MinimalAPI, config=EditorPerProjectUserSettings, PerObjectConfig, DisplayName = "Audio Input Device")
+class UTakeRecorderMicrophoneAudioManager : public UTakeRecorderAudioInputSettings
 {
 public:
 	GENERATED_BODY()
 
-	UTakeRecorderMicrophoneAudioManager(const FObjectInitializer& ObjInit);
+	UE_API UTakeRecorderMicrophoneAudioManager(const FObjectInitializer& ObjInit);
 
 	// Begin UObject Interface
-	virtual void PostEditChangeProperty(FPropertyChangedEvent& PropertyChangedEvent) override;
+	UE_API virtual void PostEditChangeProperty(FPropertyChangedEvent& PropertyChangedEvent) override;
 	// End UObject Interface
 
 	/** Enumerates the audio devices present on the current machine */
-	virtual void EnumerateAudioDevices(bool InForceRefresh = false) override;
+	UE_API virtual void EnumerateAudioDevices(bool InForceRefresh = false) override;
 
 	/** Returns input channel count for currently selected audio device */
-	virtual int32 GetDeviceChannelCount() override;
+	UE_API virtual int32 GetDeviceChannelCount() override;
 
 	/** The audio device to use for this microphone source */
 	UPROPERTY(config, EditAnywhere, BlueprintReadWrite, Category = "Source", meta = (ShowOnlyInnerProperties))
@@ -43,21 +45,21 @@ public:
 	*  Calls StartRecording on the AudioRecorder object. This is called multiple times (once for each
 	*  microphone source, however, only the first call triggers the call to AudioRecorder.
 	*/
-	void StartRecording(int32 InChannelCount);
+	UE_API void StartRecording(int32 InChannelCount);
 	/**
 	 * Calls StopRecording on the AudioRecorder object. This is called multiple times (once for each
 	 * microphone source, however, only the first call triggers the call to AudioRecorder.
 	 */
-	void StopRecording();
+	UE_API void StopRecording();
 	/**
 	 *  Calls FinalizeRecording on the AudioRecorder object. This is called multiple times (once for each
 	 *  microphone source, however, only the first call triggers the call to AudioRecorder.
 	 */
-	void FinalizeRecording();
+	UE_API void FinalizeRecording();
 
 
 	/** Fetches the USoundWave for this source after a Take has been recorded */
-	TObjectPtr<class USoundWave> GetRecordedSoundWave(const FTakeRecorderAudioSourceSettings& InSourceSettings);
+	UE_API TObjectPtr<class USoundWave> GetRecordedSoundWave(const FTakeRecorderAudioSourceSettings& InSourceSettings);
 	/** Accessor for the OnNotifySourcesOfDeviceChange delegate list */
 	FOnNotifySourcesOfDeviceChange& GetOnNotifySourcesOfDeviceChange() { return OnNotifySourcesOfDeviceChange; }
 
@@ -66,13 +68,15 @@ private:
 	/** Multicast delegate which notifies clients when the currently selected audio device changes. */
 	FOnNotifySourcesOfDeviceChange OnNotifySourcesOfDeviceChange;
 	/** Calls factory to create the AudioRecorder object */
-	TUniquePtr<IAudioCaptureEditor> CreateAudioRecorderObject();
+	UE_API TUniquePtr<IAudioCaptureEditor> CreateAudioRecorderObject();
 	/** Builds the list of audio devices which will be used in the device menu */
-	void BuildDeviceInfoArray();
+	UE_API void BuildDeviceInfoArray();
 
 	/** Returns whether audio device with given Id was found during enumeration */
-	bool IsAudioDeviceAvailable(const FString& InDeviceId);
+	UE_API bool IsAudioDeviceAvailable(const FString& InDeviceId);
 
 	/** The audio recorder object which manages low level recording of audio data */
 	TUniquePtr<IAudioCaptureEditor> AudioRecorder;
 };
+
+#undef UE_API
diff --git a/Engine/Plugins/VirtualProduction/Takes/Source/TakeRecorderSources/Private/TakeRecorderMicrophoneAudioSource.h b/Engine/Plugins/VirtualProduction/Takes/Source/TakeRecorderSources/Private/TakeRecorderMicrophoneAudioSource.h
--- a/Engine/Plugins/VirtualProduction/Takes/Source/TakeRecorderSources/Private/TakeRecorderMicrophoneAudioSource.h
+++ b/Engine/Plugins/VirtualProduction/Takes/Source/TakeRecorderSources/Private/TakeRecorderMicrophoneAudioSource.h
@@ -12,25 +12,27 @@
 
 #include "TakeRecorderMicrophoneAudioSource.generated.h"
 
+#define UE_API TAKERECORDERSOURCES_API
+
 class UMovieSceneAudioTrack;
 class USoundWave;
 class UTakeRecorderMicrophoneAudioManager;
 struct FTakeRecorderAudioSourceSettings;
 
 /** A recording source that records microphone audio */
-UCLASS(Abstract, config=EditorSettings, DisplayName="Microphone Audio Recorder")
-class TAKERECORDERSOURCES_API UTakeRecorderMicrophoneAudioSourceSettings : public UTakeRecorderSource
+UCLASS(MinimalAPI, Abstract, config=EditorSettings, DisplayName="Microphone Audio Recorder")
+class UTakeRecorderMicrophoneAudioSourceSettings : public UTakeRecorderSource
 {
 public:
 	GENERATED_BODY()
 
-	UTakeRecorderMicrophoneAudioSourceSettings(const FObjectInitializer& ObjInit);
+	UE_API UTakeRecorderMicrophoneAudioSourceSettings(const FObjectInitializer& ObjInit);
 
-	virtual void PostEditChangeProperty(FPropertyChangedEvent& PropertyChangedEvent) override;
+	UE_API virtual void PostEditChangeProperty(FPropertyChangedEvent& PropertyChangedEvent) override;
 	
 	// UTakeRecorderSource Interface
-	virtual FString GetSubsceneTrackName(ULevelSequence* InSequence) const override;
-	virtual FString GetSubsceneAssetName(ULevelSequence* InSequence) const override;
+	UE_API virtual FString GetSubsceneTrackName(ULevelSequence* InSequence) const override;
+	UE_API virtual FString GetSubsceneAssetName(ULevelSequence* InSequence) const override;
 	// ~UTakeRecorderSource Interface
 
 	/** Name of the audio source */
@@ -71,16 +73,16 @@ public:
 };
 
 /** A recording source that records microphone audio */
-UCLASS(Category="Audio", config=EditorSettings)
-class TAKERECORDERSOURCES_API UTakeRecorderMicrophoneAudioSource : public UTakeRecorderMicrophoneAudioSourceSettings
+UCLASS(MinimalAPI, Category="Audio", config=EditorSettings)
+class UTakeRecorderMicrophoneAudioSource : public UTakeRecorderMicrophoneAudioSourceSettings
 {
 public:
 	GENERATED_BODY()
 
-	UTakeRecorderMicrophoneAudioSource(const FObjectInitializer& ObjInit);
+	UE_API UTakeRecorderMicrophoneAudioSource(const FObjectInitializer& ObjInit);
 
 	// Begin UObject Interface
-	virtual void PostEditChangeProperty(FPropertyChangedEvent& PropertyChangedEvent) override;
+	UE_API virtual void PostEditChangeProperty(FPropertyChangedEvent& PropertyChangedEvent) override;
 	// End UObject Interface
 
 	/** Gain in decibels to apply to recorded audio */
@@ -102,58 +104,58 @@ public:
 	FAudioInputDeviceChannelProperty AudioChannel;
 
 	// UTakeRecorderSource
-	virtual void Initialize() override;
-	virtual FString GetSubsceneTrackName(ULevelSequence* InSequence) const override;
-	virtual FString GetSubsceneAssetName(ULevelSequence* InSequence) const override;
-	virtual TArray<UTakeRecorderSource*> PreRecording(ULevelSequence* InSequence, FMovieSceneSequenceID InSequenceID, ULevelSequence* InRootSequence, FManifestSerializer* InManifestSerializer) override;
-	virtual void StartRecording(const FTimecode& InSectionStartTimecode, const FFrameNumber& InSectionFirstFrame, class ULevelSequence* InSequence) override;
-	virtual void TickRecording(const FQualifiedFrameTime& CurrentTime) override;
-	virtual void StopRecording(class ULevelSequence* InSequence) override;
-	virtual TArray<UTakeRecorderSource*> PostRecording(ULevelSequence* InSequence, class ULevelSequence* InRootSequence, const bool bCancelled) override;
-	virtual void FinalizeRecording() override;
-	virtual FText GetDisplayTextImpl() const override;
-	virtual FText GetAddSourceDisplayTextImpl() const override;
-
-	virtual void AddContentsToFolder(class UMovieSceneFolder* InFolder) override;
-	virtual bool CanAddSource(UTakeRecorderSources* InSources) const override;
+	UE_API virtual void Initialize() override;
+	UE_API virtual FString GetSubsceneTrackName(ULevelSequence* InSequence) const override;
+	UE_API virtual FString GetSubsceneAssetName(ULevelSequence* InSequence) const override;
+	UE_API virtual TArray<UTakeRecorderSource*> PreRecording(ULevelSequence* InSequence, FMovieSceneSequenceID InSequenceID, ULevelSequence* InRootSequence, FManifestSerializer* InManifestSerializer) override;
+	UE_API virtual void StartRecording(const FTimecode& InSectionStartTimecode, const FFrameNumber& InSectionFirstFrame, class ULevelSequence* InSequence) override;
+	UE_API virtual void TickRecording(const FQualifiedFrameTime& CurrentTime) override;
+	UE_API virtual void StopRecording(class ULevelSequence* InSequence) override;
+	UE_API virtual TArray<UTakeRecorderSource*> PostRecording(ULevelSequence* InSequence, class ULevelSequence* InRootSequence, const bool bCancelled) override;
+	UE_API virtual void FinalizeRecording() override;
+	UE_API virtual FText GetDisplayTextImpl() const override;
+	UE_API virtual FText GetAddSourceDisplayTextImpl() const override;
+
+	UE_API virtual void AddContentsToFolder(class UMovieSceneFolder* InFolder) override;
+	UE_API virtual bool CanAddSource(UTakeRecorderSources* InSources) const override;
 	// ~UTakeRecorderSource
 
 	/** Returns track display name, replacing any tokens if needed */
-	FString GetAudioTrackName(ULevelSequence* InSequence) const;
+	UE_API FString GetAudioTrackName(ULevelSequence* InSequence) const;
 	/** Returns the fully expanded asset name */
-	FString GetAudioAssetName(ULevelSequence* InSequence) const;
+	UE_API FString GetAudioAssetName(ULevelSequence* InSequence) const;
 
 	/** Sets the channel count supported by the currently selected audio device. */
-	void SetAudioDeviceChannelCount(int32 InChannelCount);
+	UE_API void SetAudioDeviceChannelCount(int32 InChannelCount);
 	/** Delegate which receives notifications when the audio input device changes. */
-	void OnNotifySourcesOfDeviceChange(int32 InChannelCount);
+	UE_API void OnNotifySourcesOfDeviceChange(int32 InChannelCount);
 
 private:
 
 	/** Helper function for getting a pointer to the AudioInputManger object. */
-	static UTakeRecorderMicrophoneAudioManager* GetAudioInputManager();
+	static UE_API UTakeRecorderMicrophoneAudioManager* GetAudioInputManager();
 	
 	/** Parses recorded timecode values */
-	void ProcessRecordedTimes(ULevelSequence* InSequence);
+	UE_API void ProcessRecordedTimes(ULevelSequence* InSequence);
 
 	/** Returns an array of booleans indicating which channel indexes are currently in use */
-	TUniquePtr<TArray<bool>> GetChannelsInUse(const int32 InDeviceChannelCount);
+	UE_API TUniquePtr<TArray<bool>> GetChannelsInUse(const int32 InDeviceChannelCount);
 	/** Fetches the USoundWave for this source after a Take has been recorded */
-	void GetRecordedSoundWave(ULevelSequence* InSequence);
+	UE_API void GetRecordedSoundWave(ULevelSequence* InSequence);
 	/** Called when the user selects a new channel in the combobox for this source */
-	void SetCurrentInputChannel(int32 InChannelNumber);
+	UE_API void SetCurrentInputChannel(int32 InChannelNumber);
 
 	/** Returns whether given track is being used by any other source */
-	bool IsTrackAssociatedWithAnySource(UMovieSceneAudioTrack* InAudioTrack);
+	UE_API bool IsTrackAssociatedWithAnySource(UMovieSceneAudioTrack* InAudioTrack);
 
 	/** Build asset name, appending channel name if needed for uniqueness */
-	FString GetUniqueAudioAssetName(ULevelSequence* InSequence) const;
+	UE_API FString GetUniqueAudioAssetName(ULevelSequence* InSequence) const;
 
 	/** Static, parameterized helper function for buidling unique asset names */	
-	static FString CreateUniqueAudioAssetName(ULevelSequence* InSequence, UTakeRecorderSources* InSources, const FString& InAssetName, const int32 InChannelNumber);
+	static UE_API FString CreateUniqueAudioAssetName(ULevelSequence* InSequence, UTakeRecorderSources* InSources, const FString& InAssetName, const int32 InChannelNumber);
 
 	/** Evaluate our tokens and automatically create an appropriate context. */
-	bool TryEvaluateTokenString(const FString& InString, const ULevelSequence* InSequence, FString& OutString) const;
+	UE_API bool TryEvaluateTokenString(const FString& InString, const ULevelSequence* InSequence, FString& OutString) const;
 	
 private:
 
@@ -169,3 +171,5 @@ private:
 	// The name of the UsoundWave asset 
 	FString AssetFileName;
 };
+
+#undef UE_API
diff --git a/Engine/Plugins/VirtualProduction/Takes/Source/TakeTrackRecorders/Public/TrackRecorders/MovieScene3DAttachTrackRecorder.h b/Engine/Plugins/VirtualProduction/Takes/Source/TakeTrackRecorders/Public/TrackRecorders/MovieScene3DAttachTrackRecorder.h
--- a/Engine/Plugins/VirtualProduction/Takes/Source/TakeTrackRecorders/Public/TrackRecorders/MovieScene3DAttachTrackRecorder.h
+++ b/Engine/Plugins/VirtualProduction/Takes/Source/TakeTrackRecorders/Public/TrackRecorders/MovieScene3DAttachTrackRecorder.h
@@ -10,16 +10,18 @@
 #include "IMovieSceneTrackRecorderFactory.h"
 #include "MovieScene3DAttachTrackRecorder.generated.h"
 
+#define UE_API TAKETRACKRECORDERS_API
+
 class UMovieScene3DAttachTrack;
 class UMovieScene3DAttachSection;
 
-class TAKETRACKRECORDERS_API FMovieScene3DAttachTrackRecorderFactory : public IMovieSceneTrackRecorderFactory
+class FMovieScene3DAttachTrackRecorderFactory : public IMovieSceneTrackRecorderFactory
 {
 public:
 	virtual ~FMovieScene3DAttachTrackRecorderFactory() {}
 
-	virtual bool CanRecordObject(class UObject* InObjectToRecord) const override;
-	virtual UMovieSceneTrackRecorder* CreateTrackRecorderForObject() const override;
+	UE_API virtual bool CanRecordObject(class UObject* InObjectToRecord) const override;
+	UE_API virtual UMovieSceneTrackRecorder* CreateTrackRecorderForObject() const override;
 
 	// Attachment isn't based on any particular property
 	virtual bool CanRecordProperty(class UObject* InObjectToRecord, class FProperty* InPropertyToRecord) const override { return false; }
@@ -28,14 +30,14 @@ public:
 	virtual FText GetDisplayName() const override { return NSLOCTEXT("MovieScene3DAttachTrackRecorderFactory", "DisplayName", "Attach Track"); }
 };
 
-UCLASS(BlueprintType)
-class TAKETRACKRECORDERS_API UMovieScene3DAttachTrackRecorder : public UMovieSceneTrackRecorder
+UCLASS(MinimalAPI, BlueprintType)
+class UMovieScene3DAttachTrackRecorder : public UMovieSceneTrackRecorder
 {
 	GENERATED_BODY()
 protected:
 	// UMovieSceneTrackRecorder Interface
-	virtual void RecordSampleImpl(const FQualifiedFrameTime& CurrentTime) override;
-	virtual void FinalizeTrackImpl() override;
+	UE_API virtual void RecordSampleImpl(const FQualifiedFrameTime& CurrentTime) override;
+	UE_API virtual void FinalizeTrackImpl() override;
 	// ~UMovieSceneTrackRecorder Interface
 
 private:
@@ -50,3 +52,5 @@ private:
 	/** Track the actor we are attached to */
 	TLazyObjectPtr<class AActor> ActorAttachedTo;
 };
+
+#undef UE_API
diff --git a/Engine/Plugins/VirtualProduction/Takes/Source/TakeTrackRecorders/Public/TrackRecorders/MovieScene3DTransformTrackRecorder.h b/Engine/Plugins/VirtualProduction/Takes/Source/TakeTrackRecorders/Public/TrackRecorders/MovieScene3DTransformTrackRecorder.h
--- a/Engine/Plugins/VirtualProduction/Takes/Source/TakeTrackRecorders/Public/TrackRecorders/MovieScene3DTransformTrackRecorder.h
+++ b/Engine/Plugins/VirtualProduction/Takes/Source/TakeTrackRecorders/Public/TrackRecorders/MovieScene3DTransformTrackRecorder.h
@@ -12,6 +12,8 @@
 #include "TakesCoreLog.h"
 #include "MovieScene3DTransformTrackRecorder.generated.h"
 
+#define UE_API TAKETRACKRECORDERS_API
+
 class FMovieScene3DTransformTrackRecorder;
 class UMovieScene3DTransformTrack;
 class UMovieScene3DTransformTrack;
@@ -21,13 +23,13 @@ struct FActorRecordingSettings;
 DECLARE_LOG_CATEGORY_EXTERN(TransformSerialization, Verbose, All);
 
 
-class TAKETRACKRECORDERS_API FMovieScene3DTransformTrackRecorderFactory : public IMovieSceneTrackRecorderFactory
+class FMovieScene3DTransformTrackRecorderFactory : public IMovieSceneTrackRecorderFactory
 {
 public:
 	virtual ~FMovieScene3DTransformTrackRecorderFactory() {}
 
-	virtual bool CanRecordObject(class UObject* InObjectToRecord) const override;
-	virtual UMovieSceneTrackRecorder* CreateTrackRecorderForObject() const override;
+	UE_API virtual bool CanRecordObject(class UObject* InObjectToRecord) const override;
+	UE_API virtual UMovieSceneTrackRecorder* CreateTrackRecorderForObject() const override;
 
 	// takerecorder-todo: This should also record "Transform" variable properties, because they can be marked as interp
 	virtual bool CanRecordProperty(class UObject* InObjectToRecord, class FProperty* InPropertyToRecord) const override { return false; }
@@ -277,8 +279,8 @@ struct FBufferedTransformKeys
 };
 }; // namespace UE::TrackRecorders
 
-UCLASS(BlueprintType)
-class TAKETRACKRECORDERS_API UMovieScene3DTransformTrackRecorder : public UMovieSceneTrackRecorder
+UCLASS(MinimalAPI, BlueprintType)
+class UMovieScene3DTransformTrackRecorder : public UMovieSceneTrackRecorder
 {
 	GENERATED_BODY()
 protected:
@@ -289,10 +291,10 @@ protected:
 	}
 
 	// UMovieSceneTrackRecorder Interface
-	virtual void CreateTrackImpl() override;
-	virtual void FinalizeTrackImpl() override;
-	virtual void StopRecordingImpl() override;
-	virtual void RecordSampleImpl(const FQualifiedFrameTime& CurrentTime) override;
+	UE_API virtual void CreateTrackImpl() override;
+	UE_API virtual void FinalizeTrackImpl() override;
+	UE_API virtual void StopRecordingImpl() override;
+	UE_API virtual void RecordSampleImpl(const FQualifiedFrameTime& CurrentTime) override;
 	virtual UMovieSceneSection* GetMovieSceneSection() const override { return Cast<UMovieSceneSection>(MovieSceneSection.Get()); }
 	// ~UMovieSceneTrackRecorder Interface
 
@@ -301,12 +303,12 @@ public:
 	{
 		TransformSerializer.SetLocalCaptureDir(InDirectory);
 	}
-	virtual bool LoadRecordedFile(const FString& InFileName, UMovieScene *InMovieScene, TMap<FGuid, AActor*>& ActorGuidToActorMap,  TFunction<void()> InCompletionCallback) override;
-	void PostProcessAnimationData(class UMovieSceneAnimationTrackRecorder* AnimTrackRecorder);
+	UE_API virtual bool LoadRecordedFile(const FString& InFileName, UMovieScene *InMovieScene, TMap<FGuid, AActor*>& ActorGuidToActorMap,  TFunction<void()> InCompletionCallback) override;
+	UE_API void PostProcessAnimationData(class UMovieSceneAnimationTrackRecorder* AnimTrackRecorder);
 private:
-	bool ResolveTransformToRecord(FTransform& TransformToRecord);
-	bool ShouldAddNewKey(const FTransform& TransformToRecord);
-	void SetUpDefaultTransform();
+	UE_API bool ResolveTransformToRecord(FTransform& TransformToRecord);
+	UE_API bool ShouldAddNewKey(const FTransform& TransformToRecord);
+	UE_API void SetUpDefaultTransform();
 private:
 	/** Track to record to */
 	TWeakObjectPtr<class UMovieScene3DTransformTrack> MovieSceneTrack;
@@ -334,3 +336,5 @@ private:
 	/**Serializer */
 	FTransformSerializer TransformSerializer;
 };
+
+#undef UE_API
diff --git a/Engine/Plugins/VirtualProduction/Takes/Source/TakeTrackRecorders/Public/TrackRecorders/MovieSceneAnimationTrackRecorder.h b/Engine/Plugins/VirtualProduction/Takes/Source/TakeTrackRecorders/Public/TrackRecorders/MovieSceneAnimationTrackRecorder.h
--- a/Engine/Plugins/VirtualProduction/Takes/Source/TakeTrackRecorders/Public/TrackRecorders/MovieSceneAnimationTrackRecorder.h
+++ b/Engine/Plugins/VirtualProduction/Takes/Source/TakeTrackRecorders/Public/TrackRecorders/MovieSceneAnimationTrackRecorder.h
@@ -13,19 +13,21 @@
 #include "AnimationRecorder.h"
 #include "MovieSceneAnimationTrackRecorder.generated.h"
 
+#define UE_API TAKETRACKRECORDERS_API
+
 class FMovieScene3DTransformTrackRecorder;
 class UMovieScene3DTransformTrack;
 class UMovieScene3DTransformTrack;
 class UTakeRecorderActorSource;
 struct FActorRecordingSettings;
 
-class TAKETRACKRECORDERS_API FMovieSceneAnimationTrackRecorderFactory : public IMovieSceneTrackRecorderFactory
+class FMovieSceneAnimationTrackRecorderFactory : public IMovieSceneTrackRecorderFactory
 {
 public:
 	virtual ~FMovieSceneAnimationTrackRecorderFactory() {}
 
-	virtual bool CanRecordObject(class UObject* InObjectToRecord) const override;
-	virtual UMovieSceneTrackRecorder* CreateTrackRecorderForObject() const override;
+	UE_API virtual bool CanRecordObject(class UObject* InObjectToRecord) const override;
+	UE_API virtual UMovieSceneTrackRecorder* CreateTrackRecorderForObject() const override;
 
 	virtual bool CanRecordProperty(class UObject* InObjectToRecord, class FProperty* InPropertyToRecord) const override { return false; }
 	virtual UMovieSceneTrackRecorder* CreateTrackRecorderForProperty(UObject* InObjectToRecord, const FName& InPropertyToRecord) const override { return nullptr; }
@@ -37,33 +39,33 @@ public:
 	virtual FName GetSerializedType() const override { return FName("Animation"); }
 };
 
-UCLASS(BlueprintType)
-class TAKETRACKRECORDERS_API UMovieSceneAnimationTrackRecorder : public UMovieSceneTrackRecorder
+UCLASS(MinimalAPI, BlueprintType)
+class UMovieSceneAnimationTrackRecorder : public UMovieSceneTrackRecorder
 {
 	GENERATED_BODY()
 protected:
 	// UMovieSceneTrackRecorder Interface
-	virtual void CreateTrackImpl() override;
-	virtual void FinalizeTrackImpl() override;
-	virtual void CancelTrackImpl() override;
-	virtual void RecordSampleImpl(const FQualifiedFrameTime& CurrentTime) override;
-	virtual void StopRecordingImpl() override;
+	UE_API virtual void CreateTrackImpl() override;
+	UE_API virtual void FinalizeTrackImpl() override;
+	UE_API virtual void CancelTrackImpl() override;
+	UE_API virtual void RecordSampleImpl(const FQualifiedFrameTime& CurrentTime) override;
+	UE_API virtual void StopRecordingImpl() override;
 	virtual void SetSavedRecordingDirectory(const FString& InDirectory) override
 	{
 		AnimationSerializer.SetLocalCaptureDir(InDirectory);
 	}
-	virtual bool LoadRecordedFile(const FString& InFileName, UMovieScene *InMovieScene, TMap<FGuid, AActor*>& ActorGuidToActorMap, TFunction<void()> InCompletionCallback) override;
+	UE_API virtual bool LoadRecordedFile(const FString& InFileName, UMovieScene *InMovieScene, TMap<FGuid, AActor*>& ActorGuidToActorMap, TFunction<void()> InCompletionCallback) override;
 	virtual UMovieSceneSection* GetMovieSceneSection() const override { return Cast<UMovieSceneSection>(MovieSceneSection.Get()); }
 	// UMovieSceneTrackRecorder Interface
 
 public:
 	bool RootWasRemoved() const { return bRootWasRemoved; }
-	void RemoveRootMotion();
+	UE_API void RemoveRootMotion();
 	UE_DEPRECATED(5.5, "Use the ProcessRecordedTimes method that takes a FProcessRecordedTimeParams struct.")
-	void ProcessRecordedTimes(const FString& HoursName, const FString& MinutesName, const FString& SecondsName, const FString& FramesName, const FString& SubFramesName, const FString& SlateName, const FString& Slate);
+	UE_API void ProcessRecordedTimes(const FString& HoursName, const FString& MinutesName, const FString& SecondsName, const FString& FramesName, const FString& SubFramesName, const FString& SlateName, const FString& Slate);
 
 	/** Apply time code data to the animation track. */
-	void ProcessRecordedTimes(const FProcessRecordedTimeParams& InParams);
+	UE_API void ProcessRecordedTimes(const FProcessRecordedTimeParams& InParams);
 
 	UAnimSequence* GetAnimSequence() const { return AnimSequence.Get(); }
 	USkeletalMesh* GetSkeletalMesh() const { return SkeletalMesh.Get(); }
@@ -72,8 +74,8 @@ public:
 	const FTransform& GetInitialRootTransform() const { return InitialRootTransform; }
 
 private:
-	bool ResolveTransformToRecord(FTransform& TransformToRecord);
-	void CreateAnimationAssetAndSequence(const AActor* Actor, const FDirectoryPath& AnimationDirectory);
+	UE_API bool ResolveTransformToRecord(FTransform& TransformToRecord);
+	UE_API void CreateAnimationAssetAndSequence(const AActor* Actor, const FDirectoryPath& AnimationDirectory);
 private:
 	/** Section to record to */
 	TWeakObjectPtr<UMovieSceneSkeletalAnimationSection> MovieSceneSection;
@@ -104,3 +106,5 @@ private:
 	/** Root Was Removed*/
 	bool bRootWasRemoved = true;
 };
+
+#undef UE_API
diff --git a/Engine/Plugins/VirtualProduction/Takes/Source/TakeTrackRecorders/Public/TrackRecorders/MovieSceneAnimationTrackRecorderSettings.h b/Engine/Plugins/VirtualProduction/Takes/Source/TakeTrackRecorders/Public/TrackRecorders/MovieSceneAnimationTrackRecorderSettings.h
--- a/Engine/Plugins/VirtualProduction/Takes/Source/TakeTrackRecorders/Public/TrackRecorders/MovieSceneAnimationTrackRecorderSettings.h
+++ b/Engine/Plugins/VirtualProduction/Takes/Source/TakeTrackRecorders/Public/TrackRecorders/MovieSceneAnimationTrackRecorderSettings.h
@@ -9,14 +9,16 @@
 #include "AnimationRecorder.h"
 #include "MovieSceneAnimationTrackRecorderSettings.generated.h"
 
-UCLASS(Abstract, BlueprintType, config=EditorSettings, DisplayName="Animation Recorder")
-class TAKETRACKRECORDERS_API UMovieSceneAnimationTrackRecorderEditorSettings : public UMovieSceneTrackRecorderSettings
+#define UE_API TAKETRACKRECORDERS_API
+
+UCLASS(MinimalAPI, Abstract, BlueprintType, config=EditorSettings, DisplayName="Animation Recorder")
+class UMovieSceneAnimationTrackRecorderEditorSettings : public UMovieSceneTrackRecorderSettings
 {
 	GENERATED_BODY()
 public:
-	UMovieSceneAnimationTrackRecorderEditorSettings(const FObjectInitializer& ObjInit);
+	UE_API UMovieSceneAnimationTrackRecorderEditorSettings(const FObjectInitializer& ObjInit);
 
-	virtual void PostEditChangeProperty(FPropertyChangedEvent& PropertyChangedEvent) override;
+	UE_API virtual void PostEditChangeProperty(FPropertyChangedEvent& PropertyChangedEvent) override;
 
 	/** Name of the recorded animation track. */
 	UPROPERTY(config, EditAnywhere, BlueprintReadWrite, Category = "Animation Recorder Settings", meta = (NamingTokens))
@@ -69,8 +71,8 @@ public:
 	FTimecodeBoneMethod TimecodeBoneMethod;
 };
 
-UCLASS(BlueprintType, config = EditorSettings, DisplayName = "Animation Recorder Settings")
-class TAKETRACKRECORDERS_API UMovieSceneAnimationTrackRecorderSettings : public UMovieSceneAnimationTrackRecorderEditorSettings
+UCLASS(MinimalAPI, BlueprintType, config = EditorSettings, DisplayName = "Animation Recorder Settings")
+class UMovieSceneAnimationTrackRecorderSettings : public UMovieSceneAnimationTrackRecorderEditorSettings
 {
 	GENERATED_BODY()
 public:
@@ -78,4 +80,6 @@ public:
 		: Super(ObjInit)
 	{
 	}
-};
\ No newline at end of file
+};
+
+#undef UE_API
diff --git a/Engine/Plugins/VirtualProduction/Takes/Source/TakeTrackRecorders/Public/TrackRecorders/MovieSceneParticleTrackRecorder.h b/Engine/Plugins/VirtualProduction/Takes/Source/TakeTrackRecorders/Public/TrackRecorders/MovieSceneParticleTrackRecorder.h
--- a/Engine/Plugins/VirtualProduction/Takes/Source/TakeTrackRecorders/Public/TrackRecorders/MovieSceneParticleTrackRecorder.h
+++ b/Engine/Plugins/VirtualProduction/Takes/Source/TakeTrackRecorders/Public/TrackRecorders/MovieSceneParticleTrackRecorder.h
@@ -12,15 +12,17 @@
 #include "MovieSceneTrackRecorder.h"
 #include "MovieSceneParticleTrackRecorder.generated.h"
 
+#define UE_API TAKETRACKRECORDERS_API
+
 class UMovieSceneParticleTrackRecorder;
 
-class TAKETRACKRECORDERS_API FMovieSceneParticleTrackRecorderFactory : public IMovieSceneTrackRecorderFactory
+class FMovieSceneParticleTrackRecorderFactory : public IMovieSceneTrackRecorderFactory
 {
 public:
 	virtual ~FMovieSceneParticleTrackRecorderFactory() {}
 
-	virtual bool CanRecordObject(class UObject* InObjectToRecord) const override;
-	virtual UMovieSceneTrackRecorder* CreateTrackRecorderForObject() const override;
+	UE_API virtual bool CanRecordObject(class UObject* InObjectToRecord) const override;
+	UE_API virtual UMovieSceneTrackRecorder* CreateTrackRecorderForObject() const override;
 
 	// Particle Systems are entire components and you can't animate them as a property
 	virtual bool CanRecordProperty(class UObject* InObjectToRecord, class FProperty* InPropertyToRecord) const override { return false; }
@@ -29,8 +31,8 @@ public:
 	virtual FText GetDisplayName() const override { return NSLOCTEXT("MovieSceneParticleTrackRecorderFactory", "DisplayName", "Particle System Track"); }
 };
 
-UCLASS(BlueprintType)
-class TAKETRACKRECORDERS_API UMovieSceneParticleTrackRecorder : public UMovieSceneTrackRecorder
+UCLASS(MinimalAPI, BlueprintType)
+class UMovieSceneParticleTrackRecorder : public UMovieSceneTrackRecorder
 {
 	GENERATED_BODY()
 public:
@@ -42,14 +44,14 @@ public:
 
 protected:
 	// UMovieSceneTrackRecorder Interface
-	virtual void CreateTrackImpl() override;
-	virtual void RecordSampleImpl(const FQualifiedFrameTime& CurrentTime) override;
+	UE_API virtual void CreateTrackImpl() override;
+	UE_API virtual void RecordSampleImpl(const FQualifiedFrameTime& CurrentTime) override;
 	virtual UMovieSceneSection* GetMovieSceneSection() const override { return MovieSceneSection.Get(); }
 	// UMovieSceneTrackRecorder Interface
 
 private:
 	UFUNCTION()
-	void OnTriggered(UParticleSystemComponent* Component, bool bActivating);
+	UE_API void OnTriggered(UParticleSystemComponent* Component, bool bActivating);
 private:
 	/** Object to record from */
 	TLazyObjectPtr<class UParticleSystemComponent> SystemToRecord;
@@ -61,3 +63,5 @@ private:
 
 	EParticleKey PreviousState;
 };
+
+#undef UE_API
diff --git a/Engine/Plugins/VirtualProduction/Takes/Source/TakeTrackRecorders/Public/TrackRecorders/MovieScenePropertyTrackRecorder.h b/Engine/Plugins/VirtualProduction/Takes/Source/TakeTrackRecorders/Public/TrackRecorders/MovieScenePropertyTrackRecorder.h
--- a/Engine/Plugins/VirtualProduction/Takes/Source/TakeTrackRecorders/Public/TrackRecorders/MovieScenePropertyTrackRecorder.h
+++ b/Engine/Plugins/VirtualProduction/Takes/Source/TakeTrackRecorders/Public/TrackRecorders/MovieScenePropertyTrackRecorder.h
@@ -10,10 +10,12 @@
 #include "Serializers/MovieScenePropertySerialization.h"
 #include "MovieScenePropertyTrackRecorder.generated.h"
 
+#define UE_API TAKETRACKRECORDERS_API
+
 // Forward Declare
 class UObject;
 
-class TAKETRACKRECORDERS_API FMovieScenePropertyTrackRecorderFactory : public IMovieSceneTrackRecorderFactory
+class FMovieScenePropertyTrackRecorderFactory : public IMovieSceneTrackRecorderFactory
 {
 public:
 	virtual ~FMovieScenePropertyTrackRecorderFactory() {}
@@ -22,8 +24,8 @@ public:
 	virtual bool CanRecordObject(UObject* InObjectToRecord) const override { return false; }
 	virtual UMovieSceneTrackRecorder* CreateTrackRecorderForObject() const override { return nullptr; }
 
-	virtual bool CanRecordProperty(UObject* InObjectToRecord, class FProperty* InPropertyToRecord) const override;
-	virtual class UMovieSceneTrackRecorder* CreateTrackRecorderForProperty(UObject* InObjectToRecord, const FName& InPropertyToRecord) const override;
+	UE_API virtual bool CanRecordProperty(UObject* InObjectToRecord, class FProperty* InPropertyToRecord) const override;
+	UE_API virtual class UMovieSceneTrackRecorder* CreateTrackRecorderForProperty(UObject* InObjectToRecord, const FName& InPropertyToRecord) const override;
 	
 
 	virtual FText GetDisplayName() const override { return NSLOCTEXT("MovieScenePropertyTrackTrackRecorderFactory", "DisplayName", "Property Track"); }
@@ -31,27 +33,27 @@ public:
 	virtual bool IsSerializable() const override { return true; }
 	virtual FName GetSerializedType() const override { return FName("Property"); }
 
-	UMovieSceneTrackRecorder* CreateTrackRecorderForPropertyEnum(ESerializedPropertyType ePropertyType, const FName& InPropertyToRecord) const;
+	UE_API UMovieSceneTrackRecorder* CreateTrackRecorderForPropertyEnum(ESerializedPropertyType ePropertyType, const FName& InPropertyToRecord) const;
 
 };
 
-UCLASS(BlueprintType)
-class TAKETRACKRECORDERS_API UMovieScenePropertyTrackRecorder : public UMovieSceneTrackRecorder
+UCLASS(MinimalAPI, BlueprintType)
+class UMovieScenePropertyTrackRecorder : public UMovieSceneTrackRecorder
 {
 	GENERATED_BODY()
 protected:
 	// UMovieSceneTrackRecorder Interface
-	virtual void CreateTrackImpl() override;
-	virtual void SetSectionStartTimecodeImpl(const FTimecode& InSectionStartTimecode, const FFrameNumber& InSectionFirstFrame) override;
-	virtual void FinalizeTrackImpl() override;
-	virtual void RecordSampleImpl(const FQualifiedFrameTime& CurrentTime) override;
+	UE_API virtual void CreateTrackImpl() override;
+	UE_API virtual void SetSectionStartTimecodeImpl(const FTimecode& InSectionStartTimecode, const FFrameNumber& InSectionFirstFrame) override;
+	UE_API virtual void FinalizeTrackImpl() override;
+	UE_API virtual void RecordSampleImpl(const FQualifiedFrameTime& CurrentTime) override;
 	// ~UMovieSceneTrackRecorder Interface
 
 	virtual void SetSavedRecordingDirectory(const FString& InDirectory) override
 	{
 		Directory = InDirectory;
 	}
-	virtual bool LoadRecordedFile(const FString& InFileName, UMovieScene *InMovieScene, TMap<FGuid, AActor*>& ActorGuidToActorMap, TFunction<void()> InCompletionCallback) override;
+	UE_API virtual bool LoadRecordedFile(const FString& InFileName, UMovieScene *InMovieScene, TMap<FGuid, AActor*>& ActorGuidToActorMap, TFunction<void()> InCompletionCallback) override;
 
 
 public:
@@ -64,3 +66,5 @@ public:
 	/** Cached Directory Name for serialization, used later when we create the PropertyRecorder*/
 	FString Directory;
 };
+
+#undef UE_API
diff --git a/Engine/Plugins/VirtualProduction/Takes/Source/TakeTrackRecorders/Public/TrackRecorders/MovieSceneSpawnTrackRecorder.h b/Engine/Plugins/VirtualProduction/Takes/Source/TakeTrackRecorders/Public/TrackRecorders/MovieSceneSpawnTrackRecorder.h
--- a/Engine/Plugins/VirtualProduction/Takes/Source/TakeTrackRecorders/Public/TrackRecorders/MovieSceneSpawnTrackRecorder.h
+++ b/Engine/Plugins/VirtualProduction/Takes/Source/TakeTrackRecorders/Public/TrackRecorders/MovieSceneSpawnTrackRecorder.h
@@ -10,18 +10,20 @@
 #include "Serializers/MovieSceneSpawnSerialization.h"
 #include "MovieSceneSpawnTrackRecorder.generated.h"
 
+#define UE_API TAKETRACKRECORDERS_API
+
 // Forward Declares
 class UMovieSceneBoolTrack;
 
 DECLARE_LOG_CATEGORY_EXTERN(SpawnSerialization, Verbose, All);
 
-class TAKETRACKRECORDERS_API FMovieSceneSpawnTrackRecorderFactory : public IMovieSceneTrackRecorderFactory
+class FMovieSceneSpawnTrackRecorderFactory : public IMovieSceneTrackRecorderFactory
 {
 public:
 	virtual ~FMovieSceneSpawnTrackRecorderFactory() {}
 
-	virtual bool CanRecordObject(class UObject* InObjectToRecord) const override;
-	virtual UMovieSceneTrackRecorder* CreateTrackRecorderForObject() const override;
+	UE_API virtual bool CanRecordObject(class UObject* InObjectToRecord) const override;
+	UE_API virtual UMovieSceneTrackRecorder* CreateTrackRecorderForObject() const override;
 
 	// Spawn Track is based on whether or not the recorded object still exists
 	virtual bool CanRecordProperty(class UObject* InObjectToRecord, class FProperty* InPropertyToRecord) const override { return false; }
@@ -33,20 +35,20 @@ public:
 	virtual FName GetSerializedType() const override { return FName("Spawn"); }
 };
 
-UCLASS(BlueprintType)
-class TAKETRACKRECORDERS_API UMovieSceneSpawnTrackRecorder : public UMovieSceneTrackRecorder
+UCLASS(MinimalAPI, BlueprintType)
+class UMovieSceneSpawnTrackRecorder : public UMovieSceneTrackRecorder
 {
 	GENERATED_BODY()
 protected:
 	// UMovieSceneTrackRecorder Interface
-	virtual void CreateTrackImpl() override;
-	virtual void FinalizeTrackImpl() override;
-	virtual void RecordSampleImpl(const FQualifiedFrameTime& CurrentTime) override;
+	UE_API virtual void CreateTrackImpl() override;
+	UE_API virtual void FinalizeTrackImpl() override;
+	UE_API virtual void RecordSampleImpl(const FQualifiedFrameTime& CurrentTime) override;
 	virtual void SetSavedRecordingDirectory(const FString& InDirectory) override
 	{
 		SpawnSerializer.SetLocalCaptureDir(InDirectory);
 	}
-	virtual bool LoadRecordedFile(const FString& InFileName, UMovieScene *InMovieScene, TMap<FGuid, AActor*>& ActorGuidToActorMap, TFunction<void()> InCompletionCallback) override;
+	UE_API virtual bool LoadRecordedFile(const FString& InFileName, UMovieScene *InMovieScene, TMap<FGuid, AActor*>& ActorGuidToActorMap, TFunction<void()> InCompletionCallback) override;
 	virtual UMovieSceneSection* GetMovieSceneSection() const override { return Cast<UMovieSceneSection>(MovieSceneSection.Get()); }
 	// UMovieSceneTrackRecorder Interface
 
@@ -60,3 +62,5 @@ private:
 	FSpawnSerializer SpawnSerializer;
 	bool bSetFirstKey;
 };
+
+#undef UE_API
diff --git a/Engine/Plugins/VirtualProduction/Takes/Source/TakeTrackRecorders/Public/TrackRecorders/MovieSceneTrackRecorder.h b/Engine/Plugins/VirtualProduction/Takes/Source/TakeTrackRecorders/Public/TrackRecorders/MovieSceneTrackRecorder.h
--- a/Engine/Plugins/VirtualProduction/Takes/Source/TakeTrackRecorders/Public/TrackRecorders/MovieSceneTrackRecorder.h
+++ b/Engine/Plugins/VirtualProduction/Takes/Source/TakeTrackRecorders/Public/TrackRecorders/MovieSceneTrackRecorder.h
@@ -13,8 +13,8 @@
 // Forward Declare
 class UMovieSceneTrackRecorderSettings;
 
-UCLASS(BlueprintType, Transient, Abstract)
-class TAKETRACKRECORDERS_API UMovieSceneTrackRecorder : public UObject
+UCLASS(MinimalAPI, BlueprintType, Transient, Abstract)
+class UMovieSceneTrackRecorder : public UObject
 {
 	GENERATED_BODY()
 public:
diff --git a/Engine/Plugins/VirtualProduction/Takes/Source/TakeTrackRecorders/Public/TrackRecorders/MovieSceneTrackRecorderSettings.h b/Engine/Plugins/VirtualProduction/Takes/Source/TakeTrackRecorders/Public/TrackRecorders/MovieSceneTrackRecorderSettings.h
--- a/Engine/Plugins/VirtualProduction/Takes/Source/TakeTrackRecorders/Public/TrackRecorders/MovieSceneTrackRecorderSettings.h
+++ b/Engine/Plugins/VirtualProduction/Takes/Source/TakeTrackRecorders/Public/TrackRecorders/MovieSceneTrackRecorderSettings.h
@@ -5,8 +5,8 @@
 #include "CoreMinimal.h"
 #include "MovieSceneTrackRecorderSettings.generated.h"
 
-UCLASS(BlueprintType)
-class TAKETRACKRECORDERS_API UMovieSceneTrackRecorderSettings : public UObject
+UCLASS(MinimalAPI, BlueprintType)
+class UMovieSceneTrackRecorderSettings : public UObject
 {
 	GENERATED_BODY()
 };
diff --git a/Engine/Plugins/VirtualProduction/Takes/Source/TakeTrackRecorders/Public/TrackRecorders/MovieSceneVisibilityTrackRecorder.h b/Engine/Plugins/VirtualProduction/Takes/Source/TakeTrackRecorders/Public/TrackRecorders/MovieSceneVisibilityTrackRecorder.h
--- a/Engine/Plugins/VirtualProduction/Takes/Source/TakeTrackRecorders/Public/TrackRecorders/MovieSceneVisibilityTrackRecorder.h
+++ b/Engine/Plugins/VirtualProduction/Takes/Source/TakeTrackRecorders/Public/TrackRecorders/MovieSceneVisibilityTrackRecorder.h
@@ -10,37 +10,39 @@
 #include "Sections/MovieSceneBoolSection.h"
 #include "MovieSceneVisibilityTrackRecorder.generated.h"
 
+#define UE_API TAKETRACKRECORDERS_API
+
 class UMovieSceneBoolTrack;
 
-class TAKETRACKRECORDERS_API FMovieSceneVisibilityTrackRecorderFactory : public IMovieSceneTrackRecorderFactory
+class FMovieSceneVisibilityTrackRecorderFactory : public IMovieSceneTrackRecorderFactory
 {
 public:
 	virtual ~FMovieSceneVisibilityTrackRecorderFactory() {}
 
-	virtual bool CanRecordObject(class UObject* InObjectToRecord) const override;
-	virtual UMovieSceneTrackRecorder* CreateTrackRecorderForObject() const override;
+	UE_API virtual bool CanRecordObject(class UObject* InObjectToRecord) const override;
+	UE_API virtual UMovieSceneTrackRecorder* CreateTrackRecorderForObject() const override;
 
 	// Visibility is based on a different property for components and actors, and they're not marked as interp
-	virtual bool CanRecordProperty(class UObject* InObjectToRecord, class FProperty* InPropertyToRecord) const override;
+	UE_API virtual bool CanRecordProperty(class UObject* InObjectToRecord, class FProperty* InPropertyToRecord) const override;
 	virtual UMovieSceneTrackRecorder* CreateTrackRecorderForProperty(UObject* InObjectToRecord, const FName& InPropertyToRecord) const override { return nullptr; }
 
 	virtual FText GetDisplayName() const override { return NSLOCTEXT("MovieSceneVisibilityTrackRecorderFactory", "DisplayName", "Visibility Track"); }
 };
 
-UCLASS(BlueprintType)
-class TAKETRACKRECORDERS_API UMovieSceneVisibilityTrackRecorder : public UMovieSceneTrackRecorder
+UCLASS(MinimalAPI, BlueprintType)
+class UMovieSceneVisibilityTrackRecorder : public UMovieSceneTrackRecorder
 {
 	GENERATED_BODY()
 protected:
 	// UMovieSceneTrackRecorder Interface
-	virtual void CreateTrackImpl() override;
-	virtual void RecordSampleImpl(const FQualifiedFrameTime& CurrentTime) override;
-	virtual void FinalizeTrackImpl() override;
+	UE_API virtual void CreateTrackImpl() override;
+	UE_API virtual void RecordSampleImpl(const FQualifiedFrameTime& CurrentTime) override;
+	UE_API virtual void FinalizeTrackImpl() override;
 	virtual UMovieSceneSection* GetMovieSceneSection() const override { return Cast<UMovieSceneSection>(MovieSceneSection.Get()); }
 	// UMovieSceneTrackRecorder Interface
 
-	void RemoveRedundantTracks();
-	bool IsObjectVisible() const;
+	UE_API void RemoveRedundantTracks();
+	UE_API bool IsObjectVisible() const;
 
 private:
 	/** Section to record to */
@@ -52,3 +54,5 @@ private:
 	/** Flag used to determine whether the first key needs to be set */
 	bool bSetFirstKey;
 };
+
+#undef UE_API
diff --git a/Engine/Plugins/VirtualProduction/Takes/Source/TakesCore/Public/TakeMetaData.h b/Engine/Plugins/VirtualProduction/Takes/Source/TakesCore/Public/TakeMetaData.h
--- a/Engine/Plugins/VirtualProduction/Takes/Source/TakesCore/Public/TakeMetaData.h
+++ b/Engine/Plugins/VirtualProduction/Takes/Source/TakesCore/Public/TakeMetaData.h
@@ -9,6 +9,8 @@
 #include "Misc/QualifiedFrameTime.h"
 #include "TakeMetaData.generated.h"
 
+#define UE_API TAKESCORE_API
+
 class UTakePreset;
 class UTakeRecorderNamingTokensContext;
 struct FNamingTokenResultData;
@@ -17,41 +19,41 @@ struct FNamingTokenResultData;
  * Take meta-data that is stored on ULevelSequence assets that are recorded through the Take Recorder.
  * Meta-data is retrieved through ULevelSequence::FindMetaData<UTakeMetaData>()
  */
-UCLASS(config=EditorSettings, PerObjectConfig, BlueprintType)
-class TAKESCORE_API UTakeMetaData : public UObject, public IMovieSceneMetaDataInterface
+UCLASS(MinimalAPI, config=EditorSettings, PerObjectConfig, BlueprintType)
+class UTakeMetaData : public UObject, public IMovieSceneMetaDataInterface
 {
 public:
 	GENERATED_BODY()
 
-	UTakeMetaData(const FObjectInitializer& ObjInit);
+	UE_API UTakeMetaData(const FObjectInitializer& ObjInit);
 
 public:
 
 	/** The asset registry tag that contains the slate for this meta-data */
-	static const FName AssetRegistryTag_Slate;
+	static UE_API const FName AssetRegistryTag_Slate;
 
 	/** The asset registry tag that contains the take number for this meta-data */
-	static const FName AssetRegistryTag_TakeNumber;
+	static UE_API const FName AssetRegistryTag_TakeNumber;
 
 	/** The asset registry tag that contains the timestamp for this meta-data */
-	static const FName AssetRegistryTag_Timestamp;
+	static UE_API const FName AssetRegistryTag_Timestamp;
 
 	/** The asset registry tag that contains the timecode in for this meta-data */
-	static const FName AssetRegistryTag_TimecodeIn;
+	static UE_API const FName AssetRegistryTag_TimecodeIn;
 
 	/** The asset registry tag that contains the timecode out for this meta-data */
-	static const FName AssetRegistryTag_TimecodeOut;
+	static UE_API const FName AssetRegistryTag_TimecodeOut;
 
 	/** The asset registry tag that contains the user-description for this meta-data */
-	static const FName AssetRegistryTag_Description;
+	static UE_API const FName AssetRegistryTag_Description;
 
 	/** The asset registry tag that contains the level-path for this meta-data */
-	static const FName AssetRegistryTag_LevelPath;
+	static UE_API const FName AssetRegistryTag_LevelPath;
 	
 	/**
 	 * Access the global config instance that houses default settings for take meta data for a given project
 	 */
-	static UTakeMetaData* GetConfigInstance();
+	static UE_API UTakeMetaData* GetConfigInstance();
 
 	/**
 	 * Create a new meta-data object from the project defaults
@@ -59,31 +61,31 @@ public:
 	 * @param Outer    The object to allocate the new meta-data within
 	 * @param Name     The name for the new object. Must not already exist
 	 */
-	static UTakeMetaData* CreateFromDefaults(UObject* Outer, FName Name);
+	static UE_API UTakeMetaData* CreateFromDefaults(UObject* Outer, FName Name);
 	
 	/**
 	 * Retrieve the most recent or active take metadata.
 	 */
-	static UTakeMetaData* GetMostRecentMetaData();
+	static UE_API UTakeMetaData* GetMostRecentMetaData();
 
 	/**
 	 * Update the most recent accessed metadata.
 	 */
-	static void SetMostRecentMetaData(UTakeMetaData* InMetaData);
+	static UE_API void SetMostRecentMetaData(UTakeMetaData* InMetaData);
 	
 public:
 
 	/**
 	 * Extend the default ULevelSequence asset registry tags
 	 */
-	virtual void ExtendAssetRegistryTags(FAssetRegistryTagsContext Context) const override;
+	UE_API virtual void ExtendAssetRegistryTags(FAssetRegistryTagsContext Context) const override;
 	UE_DEPRECATED(5.4, "Implement the version that takes FAssetRegistryTagsContext instead.")
 	virtual void ExtendAssetRegistryTags(TArray<FAssetRegistryTag>& OutTags) const override {}
 
 	/**
 	 * Extend the default ULevelSequence asset registry tag meta-data
 	 */
-	virtual void ExtendAssetRegistryTagMetaData(TMap<FName, FAssetRegistryTagMetadata>& OutMetadata) const override;
+	UE_API virtual void ExtendAssetRegistryTagMetaData(TMap<FName, FAssetRegistryTagMetadata>& OutMetadata) const override;
 
 public:
 
@@ -91,86 +93,86 @@ public:
 	 * Check whether this take is locked
 	 */
 	UFUNCTION(BlueprintCallable, Category="Take")
-	bool IsLocked() const;
+	UE_API bool IsLocked() const;
 
 	/**
 	 * Check if this take was recorded (as opposed
 	 * to being setup for recording)
 	 */ 
 	UFUNCTION(BlueprintCallable, Category="Take")
-	bool Recorded() const;
+	UE_API bool Recorded() const;
 
 	/**
 	 * @return The slate for this take
 	 */
 	UFUNCTION(BlueprintCallable, Category="Take")
-	const FString& GetSlate() const;
+	UE_API const FString& GetSlate() const;
 
 	/**
 	 * @return The take number for this take
 	 */
 	UFUNCTION(BlueprintCallable, Category="Take")
-	int32 GetTakeNumber() const;
+	UE_API int32 GetTakeNumber() const;
 
 	/**
 	 * @return The timestamp for this take
 	 */
 	UFUNCTION(BlueprintCallable, Category="Take")
-	FDateTime GetTimestamp() const;
+	UE_API FDateTime GetTimestamp() const;
 
 	/**
 	 * @return The timecode in for this take
 	 */
 	UFUNCTION(BlueprintCallable, Category = "Take")
-	FTimecode GetTimecodeIn() const;
+	UE_API FTimecode GetTimecodeIn() const;
 
 	/**
 	 * @return The timecode out for this take
 	 */
 	UFUNCTION(BlueprintCallable, Category = "Take")
-	FTimecode GetTimecodeOut() const;
+	UE_API FTimecode GetTimecodeOut() const;
 
 	/**
 	 * @return The duration for this take
 	 */
 	UFUNCTION(BlueprintCallable, Category="Take")
-	FFrameTime GetDuration() const;
+	UE_API FFrameTime GetDuration() const;
 
 	/**
 	 * @return The frame-rate for this take
 	 */
 	UFUNCTION(BlueprintCallable, Category="Take")
-	FFrameRate GetFrameRate();
+	UE_API FFrameRate GetFrameRate();
 
 	/**
 	 * @return The user-provided description for this take
 	 */
 	UFUNCTION(BlueprintCallable, Category="Take")
-	FString GetDescription() const;
+	UE_API FString GetDescription() const;
 
 	/**
 	 * @return The preset on which the take was originally based
 	 */
 	UFUNCTION(BlueprintCallable, Category="Take")
-	UTakePreset* GetPresetOrigin() const;
+	UE_API UTakePreset* GetPresetOrigin() const;
 
 	/**
 	 * @return The AssetPath of the Level used to create a Recorded Level Sequence
 	 */ 	
 	UFUNCTION(BlueprintCallable, Category="Take")
-	FString GetLevelPath() const;
+	UE_API FString GetLevelPath() const;
 
 	/**
 	 * @return The Map used to create this recording
 	 */ 
 	UFUNCTION(BlueprintCallable, Category="Take")
-	ULevel* GetLevelOrigin() const;
+	UE_API ULevel* GetLevelOrigin() const;
 
 	/**
 	*  @return Get if we get frame rate from time code
 	*/
 	UFUNCTION(BlueprintCallable, Category = "Take")
-	bool GetFrameRateFromTimecode() const;
+	UE_API bool GetFrameRateFromTimecode() const;
 
 public:
 
@@ -178,28 +180,28 @@ public:
 	 * Lock this take, causing it to become read-only
 	 */
 	UFUNCTION(BlueprintCallable, Category="Take")
-	void Lock();
+	UE_API void Lock();
 
 	/**
 	 * Unlock this take if it is read-only, allowing it to be modified once again
 	 */
 	UFUNCTION(BlueprintCallable, Category="Take")
-	void Unlock();
+	UE_API void Unlock();
 
 	/**
 	 * Generate the desired asset path for this take meta-data
 	 */
 	UFUNCTION(BlueprintCallable, Category="Take")
-	FString GenerateAssetPath(const FString& PathFormatString, UTakeRecorderNamingTokensContext* InContext = nullptr) const;
+	UE_API FString GenerateAssetPath(const FString& PathFormatString, UTakeRecorderNamingTokensContext* InContext = nullptr) const;
 
 	/**
 	 * Attempt to generate the root asset path.
 	 * @return true if successful, false if the path couldn't be generated.
 	 */
-	bool TryGenerateRootAssetPath(const FString& PathFormatString, FString& OutGeneratedAssetPath, FText* OutErrorMessage = nullptr, UTakeRecorderNamingTokensContext* InContext = nullptr) const;
+	UE_API bool TryGenerateRootAssetPath(const FString& PathFormatString, FString& OutGeneratedAssetPath, FText* OutErrorMessage = nullptr, UTakeRecorderNamingTokensContext* InContext = nullptr) const;
 	
 	/** Process TakeRecorder Naming Tokens. */
-	FNamingTokenResultData ProcessTokens(const FText& InText, UTakeRecorderNamingTokensContext* InContext = nullptr) const;
+	UE_API FNamingTokenResultData ProcessTokens(const FText& InText, UTakeRecorderNamingTokensContext* InContext = nullptr) const;
 	
 	/**
 	 * Set the slate for this take and reset its take number to 1
@@ -207,7 +209,7 @@ public:
 	 * @note: Only valid for takes that have not been locked
 	 */
 	UFUNCTION(BlueprintCallable, Category="Take")
-	void SetSlate(FString InSlate, bool bEmitChanged = true);
+	UE_API void SetSlate(FString InSlate, bool bEmitChanged = true);
 
 	/**
 	 * Set this take's take number. Take numbers are always clamped to be >= 1.
@@ -215,76 +217,76 @@ public:
 	 * @note: Only valid for takes that have not been locked
 	 */
 	UFUNCTION(BlueprintCallable, Category="Take")
-	void SetTakeNumber(int32 InTakeNumber, bool bEmitChanged = true);
+	UE_API void SetTakeNumber(int32 InTakeNumber, bool bEmitChanged = true);
 
 	/**
 	 * Set this take's timestamp
 	 * @note: Only valid for takes that have not been locked
 	 */
 	UFUNCTION(BlueprintCallable, Category="Take")
-	void SetTimestamp(FDateTime InTimestamp);
+	UE_API void SetTimestamp(FDateTime InTimestamp);
 
 	/**
 	 * Set this take's timecode in
 	 * @note: Only valid for takes that have not been locked
 	 */
 	UFUNCTION(BlueprintCallable, Category = "Take")
-	void SetTimecodeIn(FTimecode InTimecodeIn);
+	UE_API void SetTimecodeIn(FTimecode InTimecodeIn);
 
 	/**
 	 * Set this take's timecode out
 	 * @note: Only valid for takes that have not been locked
 	 */
 	UFUNCTION(BlueprintCallable, Category = "Take")
-	void SetTimecodeOut(FTimecode InTimecodeOut);
+	UE_API void SetTimecodeOut(FTimecode InTimecodeOut);
 
 	/**
 	 * Set this take's duration
 	 * @note: Only valid for takes that have not been locked
 	 */
 	UFUNCTION(BlueprintCallable, Category="Take")
-	void SetDuration(FFrameTime InDuration);
+	UE_API void SetDuration(FFrameTime InDuration);
 
 	/**
 	 * Set this take's frame-rate
 	 * @note: Only valid for takes that have not been locked
 	 */
 	UFUNCTION(BlueprintCallable, Category="Take")
-	void SetFrameRate(FFrameRate InFrameRate);
+	UE_API void SetFrameRate(FFrameRate InFrameRate);
 
 	/**
 	 * Set this take's user-provided description
 	 * @note: Only valid for takes that have not been locked
 	 */
 	UFUNCTION(BlueprintCallable, Category="Take")
-	void SetDescription(FString InDescription);
+	UE_API void SetDescription(FString InDescription);
 
 	/**
 	 * Set the preset on which the take is based
 	 * @note: Only valid for takes that have not been locked
 	 */
 	UFUNCTION(BlueprintCallable, Category="Take")
-	void SetPresetOrigin(UTakePreset* InPresetOrigin);
+	UE_API void SetPresetOrigin(UTakePreset* InPresetOrigin);
 
 	/**
 	 *  Set the map used to create this recording
 	 */
 	UFUNCTION(BlueprintCallable, Category="Take")
-	void SetLevelOrigin(ULevel* InLevelOrigin);
+	UE_API void SetLevelOrigin(ULevel* InLevelOrigin);
 
 	/**
 	*  Set if we get frame rate from time code
 	*/
 	UFUNCTION(BlueprintCallable, Category = "Take")
-	void SetFrameRateFromTimecode(bool InFromTimecode);
+	UE_API void SetFrameRateFromTimecode(bool InFromTimecode);
 
 	DECLARE_MULTICAST_DELEGATE_TwoParams(FOnTakeSlateChanged, const FString& SlateName, UTakeMetaData* TakeMetaData);
 	/** Multicast delegate emitted when the slate is changed. */
-	static FOnTakeSlateChanged& OnTakeSlateChanged();
+	static UE_API FOnTakeSlateChanged& OnTakeSlateChanged();
 
 	DECLARE_MULTICAST_DELEGATE_TwoParams(FOnTakeNumberChanged, int32 TakeNumber, UTakeMetaData* TakeMetaData);
 	/** Multicast delegate emitted when take number is changed. */
-	static FOnTakeNumberChanged& OnTakeNumberChanged();
+	static UE_API FOnTakeNumberChanged& OnTakeNumberChanged();
 	
 private:
 
@@ -338,3 +340,5 @@ private:
 	UPROPERTY()
 	bool bFrameRateFromTimecode;
 };
+
+#undef UE_API
diff --git a/Engine/Plugins/VirtualProduction/Takes/Source/TakesCore/Public/TakePreset.h b/Engine/Plugins/VirtualProduction/Takes/Source/TakesCore/Public/TakePreset.h
--- a/Engine/Plugins/VirtualProduction/Takes/Source/TakesCore/Public/TakePreset.h
+++ b/Engine/Plugins/VirtualProduction/Takes/Source/TakesCore/Public/TakePreset.h
@@ -5,19 +5,21 @@
 #include "UObject/Object.h"
 #include "TakePreset.generated.h"
 
+#define UE_API TAKESCORE_API
+
 class ULevelSequence;
 
 /**
  * Take preset that is stored as an asset comprising a ULevelSequence, and a set of actor recording sources
  */
-UCLASS(BlueprintType)
-class TAKESCORE_API UTakePreset : public UObject
+UCLASS(MinimalAPI, BlueprintType)
+class UTakePreset : public UObject
 {
 public:
 
 	GENERATED_BODY()
 
-	UTakePreset(const FObjectInitializer& ObjInit);
+	UE_API UTakePreset(const FObjectInitializer& ObjInit);
 
 	/**
 	 * Get this preset's level sequence that is used as a template for a new take recording
@@ -30,40 +32,40 @@ public:
 	/**
 	 * Retrieve this preset's level sequence template, creating one if necessary
 	 */
-	ULevelSequence* GetOrCreateLevelSequence();
+	UE_API ULevelSequence* GetOrCreateLevelSequence();
 
 	/**
 	 * Forcibly re-create this preset's level sequence template, even if one already exists
 	 */
-	void CreateLevelSequence();
+	UE_API void CreateLevelSequence();
 
 	/**
 	 * Copy the specified template preset into this instance. Copies the level sequence and all its recording meta-data.
 	 */
-	void CopyFrom(UTakePreset* TemplatePreset);
+	UE_API void CopyFrom(UTakePreset* TemplatePreset);
 
 	/**
 	 * Copy the specified level-sequence into this instance. Copies the level sequence and all its recording meta-data.
 	 */
-	void CopyFrom(ULevelSequence* TemplateLevelSequence);
+	UE_API void CopyFrom(ULevelSequence* TemplateLevelSequence);
 
 	/**
 	 * Bind onto an event that is triggered when this preset's level sequence has been changed
 	 */
-	FDelegateHandle AddOnLevelSequenceChanged(const FSimpleDelegate& InHandler);
+	UE_API FDelegateHandle AddOnLevelSequenceChanged(const FSimpleDelegate& InHandler);
 
 	/**
 	 * Remove a previously bound handler for the  event that is triggered when this preset's level sequence has been changed
 	 */
-	void RemoveOnLevelSequenceChanged(FDelegateHandle DelegateHandle);
+	UE_API void RemoveOnLevelSequenceChanged(FDelegateHandle DelegateHandle);
 
 	/**
 	 * Allocate the transient preset used by the take recorder.
 	 */
-	static UTakePreset* AllocateTransientPreset(UTakePreset* TemplatePreset);
+	static UE_API UTakePreset* AllocateTransientPreset(UTakePreset* TemplatePreset);
 private:
 
-	virtual void PostTransacted(const FTransactionObjectEvent& TransactionEvent) override;
+	UE_API virtual void PostTransacted(const FTransactionObjectEvent& TransactionEvent) override;
 
 	/** Instanced level sequence template that is used to define a starting point for a new take recording */
 	UPROPERTY(Instanced)
@@ -71,3 +73,5 @@ private:
 
 	FSimpleMulticastDelegate OnLevelSequenceChangedEvent;
 };
+
+#undef UE_API
diff --git a/Engine/Plugins/VirtualProduction/Takes/Source/TakesCore/Public/TakePresetSettings.h b/Engine/Plugins/VirtualProduction/Takes/Source/TakesCore/Public/TakePresetSettings.h
--- a/Engine/Plugins/VirtualProduction/Takes/Source/TakesCore/Public/TakePresetSettings.h
+++ b/Engine/Plugins/VirtualProduction/Takes/Source/TakesCore/Public/TakePresetSettings.h
@@ -1,4 +1,4 @@
-﻿// Copyright Epic Games, Inc. All Rights Reserved.
+// Copyright Epic Games, Inc. All Rights Reserved.
 
 #pragma once
 
@@ -6,6 +6,8 @@
 #include "UObject/Object.h"
 #include "TakePresetSettings.generated.h"
 
+#define UE_API TAKESCORE_API
+
 class ULevelSequence;
 
 /** Wraps the class so this can be customized by a type layout. */
@@ -24,19 +26,19 @@ struct FTakeRecorderTargetRecordClassProperty
  * These settings belong into UTakeRecorderSettings but those are in TakeRecorder, which depends on TakesCore. 
  * DisplayName is important so these settings are displayed in the same category as UTakeRecorderSettings.
  */
-UCLASS(config=EditorSettings, BlueprintType, DisplayName = "Take Recorder")
-class TAKESCORE_API UTakePresetSettings : public UObject
+UCLASS(MinimalAPI, config=EditorSettings, BlueprintType, DisplayName = "Take Recorder")
+class UTakePresetSettings : public UObject
 {
 public:
 	GENERATED_BODY()
 
-	UTakePresetSettings();
+	UE_API UTakePresetSettings();
 
 	/** @return The settings object */
-	static UTakePresetSettings* Get();
+	static UE_API UTakePresetSettings* Get();
 
 	/** @return The class that recorded sequences should have */
-	UClass* GetTargetRecordClass() const;
+	UE_API UClass* GetTargetRecordClass() const;
 
 	DECLARE_MULTICAST_DELEGATE(FOnSettingsChanged);
 	/** Invoked when any settings change. */
@@ -45,7 +47,7 @@ public:
 	static FName GetTargetRecordClassMemberName() { return GET_MEMBER_NAME_CHECKED(UTakePresetSettings, TargetRecordClass); } 
 
 #if WITH_EDITOR
-	virtual void PostEditChangeProperty(FPropertyChangedEvent& PropertyChangedEvent) override;
+	UE_API virtual void PostEditChangeProperty(FPropertyChangedEvent& PropertyChangedEvent) override;
 #endif
 
 private:
@@ -57,3 +59,5 @@ private:
 	/** Invoked when any settings change. */
 	FOnSettingsChanged OnSettingsChangedDelegate;
 };
+
+#undef UE_API
diff --git a/Engine/Plugins/VirtualProduction/Takes/Source/TakesCore/Public/TakeRecorderSource.h b/Engine/Plugins/VirtualProduction/Takes/Source/TakesCore/Public/TakeRecorderSource.h
--- a/Engine/Plugins/VirtualProduction/Takes/Source/TakesCore/Public/TakeRecorderSource.h
+++ b/Engine/Plugins/VirtualProduction/Takes/Source/TakesCore/Public/TakeRecorderSource.h
@@ -9,6 +9,8 @@
 #include "Misc/QualifiedFrameTime.h"
 #include "TakeRecorderSource.generated.h"
 
+#define UE_API TAKESCORE_API
+
 struct FMovieSceneSequenceID;
 class ULevelSequence;
 
@@ -25,13 +27,13 @@ class ULevelSequence;
  * Sources should reset their state before recording as there is not a guarantee that the object will be newly
  * created for each recording.
  */
-UCLASS(Abstract, Blueprintable, BlueprintType)
-class TAKESCORE_API UTakeRecorderSource : public UObject
+UCLASS(MinimalAPI, Abstract, Blueprintable, BlueprintType)
+class UTakeRecorderSource : public UObject
 {
 public:
 	GENERATED_BODY()
 
-	UTakeRecorderSource(const FObjectInitializer& ObjInit);
+	UE_API UTakeRecorderSource(const FObjectInitializer& ObjInit);
 
 	/** True if this source is cued for recording or not */
 	UPROPERTY(BlueprintReadWrite, Category="Source")
@@ -76,7 +78,7 @@ public:
 	* @return An array of newly created take recorder sources. Can be an empty list if no additional sources needed to be
 	* created by this source.
 	*/
-	virtual TArray<UTakeRecorderSource*> PreRecording(ULevelSequence* InSequence, FMovieSceneSequenceID InSequenceID, ULevelSequence* InRootSequence, FManifestSerializer* InManifestSerializer);
+	UE_API virtual TArray<UTakeRecorderSource*> PreRecording(ULevelSequence* InSequence, FMovieSceneSequenceID InSequenceID, ULevelSequence* InRootSequence, FManifestSerializer* InManifestSerializer);
 
 	/**
 	* This is called when the UTakeRecorderSources starts a recording, after all sources have had PreRecording called on them.
@@ -223,9 +225,11 @@ public:
 private:
 
 	/** Default implementations for UI functions */
-	virtual const FSlateBrush* GetDisplayIconImpl() const;
+	UE_API virtual const FSlateBrush* GetDisplayIconImpl() const;
 	virtual FText GetCategoryTextImpl() const { return FText(); }
 	virtual FText GetDisplayTextImpl() const { return FText(); }
 	virtual FText GetDescriptionTextImpl() const { return FText(); }
 	virtual FText GetAddSourceDisplayTextImpl() const { return FText(); };
 };
+
+#undef UE_API
diff --git a/Engine/Plugins/VirtualProduction/Takes/Source/TakesCore/Public/TakeRecorderSourceProperty.h b/Engine/Plugins/VirtualProduction/Takes/Source/TakesCore/Public/TakeRecorderSourceProperty.h
--- a/Engine/Plugins/VirtualProduction/Takes/Source/TakesCore/Public/TakeRecorderSourceProperty.h
+++ b/Engine/Plugins/VirtualProduction/Takes/Source/TakesCore/Public/TakeRecorderSourceProperty.h
@@ -6,8 +6,10 @@
 #include "TakeRecorderSources.h"
 #include "TakeRecorderSourceProperty.generated.h"
 
+#define UE_API TAKESCORE_API
+
 USTRUCT(BlueprintType)
-struct TAKESCORE_API FActorRecordedProperty
+struct FActorRecordedProperty
 {
 	GENERATED_BODY()
 
@@ -41,14 +43,14 @@ struct TAKESCORE_API FActorRecordedProperty
 * to record them. If you wish to expose a property to be recorded it needs to be marked
 * as "Interp" (C++) or "Expose to Cinematics" in Blueprints.
 */
-UCLASS(BlueprintType)
-class TAKESCORE_API UActorRecorderPropertyMap : public UObject
+UCLASS(MinimalAPI, BlueprintType)
+class UActorRecorderPropertyMap : public UObject
 {
 	GENERATED_BODY()
 
 public:
-	virtual void PostEditUndo() override;
-	virtual void PostEditChangeProperty(struct FPropertyChangedEvent& PropertyChangedEvent) override;
+	UE_API virtual void PostEditUndo() override;
+	UE_API virtual void PostEditChangeProperty(struct FPropertyChangedEvent& PropertyChangedEvent) override;
 
 	UPROPERTY(BlueprintReadOnly, VisibleAnywhere, Category = "Property")
 	TSoftObjectPtr<UObject> RecordedObject;
@@ -82,13 +84,13 @@ public:
 	}
 
 	/** Visit all properties recursively and cache the record state. */
-	void UpdateCachedValues();
+	UE_API void UpdateCachedValues();
 private:
-	int32 NumberOfPropertiesRecordedOnThis() const;
-	int32 NumberOfComponentsRecordedOnThis() const;
+	UE_API int32 NumberOfPropertiesRecordedOnThis() const;
+	UE_API int32 NumberOfComponentsRecordedOnThis() const;
 
 	/** This is called by the children to let the parent know that the number of recorded properties has changed*/
-	void ChildChanged();
+	UE_API void ChildChanged();
 
 	TWeakObjectPtr<UActorRecorderPropertyMap> Parent;
 
@@ -99,7 +101,7 @@ private:
  * Encapsulates audio device properties which are utilized by UI facing classes such as FAudioInputDeviceProperty.
  */
 USTRUCT(BlueprintType)
-struct TAKESCORE_API FAudioInputDeviceInfoProperty
+struct FAudioInputDeviceInfoProperty
 {
 	GENERATED_BODY()
 
@@ -140,7 +142,7 @@ struct TAKESCORE_API FAudioInputDeviceInfoProperty
  * utilized by the audio input device list in FAudioInputDevicePropertyCustomization.
  */
 USTRUCT(BlueprintType)
-struct TAKESCORE_API FAudioInputDeviceProperty
+struct FAudioInputDeviceProperty
 {
 	GENERATED_BODY()
 
@@ -169,7 +171,7 @@ struct TAKESCORE_API FAudioInputDeviceProperty
  * contains a local cache of the max channel count for the currently selected audio device.
  */
 USTRUCT(BlueprintType)
-struct TAKESCORE_API FAudioInputDeviceChannelProperty
+struct FAudioInputDeviceChannelProperty
 {
 	GENERATED_BODY()
 
@@ -185,8 +187,8 @@ DECLARE_MULTICAST_DELEGATE(FOnAudioInputDeviceChanged);
 /**
  * Provides access to the UI code for registering OnAudioInputDeviceChanged delegates.
  */
-UCLASS(Abstract, Blueprintable, BlueprintType)
-class TAKESCORE_API UTakeRecorderAudioInputSettings : public UObject
+UCLASS(MinimalAPI, Abstract, Blueprintable, BlueprintType)
+class UTakeRecorderAudioInputSettings : public UObject
 {
 	GENERATED_BODY()
 	
@@ -206,3 +208,5 @@ private:
 	/** Multicast notification when the user selects a new audio device */
 	FOnAudioInputDeviceChanged OnAudioInputDeviceChanged;
 };
+
+#undef UE_API
diff --git a/Engine/Plugins/VirtualProduction/Takes/Source/TakesCore/Public/TakeRecorderSources.h b/Engine/Plugins/VirtualProduction/Takes/Source/TakesCore/Public/TakeRecorderSources.h
--- a/Engine/Plugins/VirtualProduction/Takes/Source/TakesCore/Public/TakeRecorderSources.h
+++ b/Engine/Plugins/VirtualProduction/Takes/Source/TakesCore/Public/TakeRecorderSources.h
@@ -11,6 +11,8 @@
 #include "Serializers/MovieSceneManifestSerialization.h"
 #include "TakeRecorderSources.generated.h"
 
+#define UE_API TAKESCORE_API
+
 class UTakeRecorderSource;
 class UMovieSceneSubSection;
 
@@ -31,13 +33,13 @@ struct FTakeRecorderSourcesSettings
 /**
  * A list of sources to record for any given take. Stored as meta-data on ULevelSequence through ULevelSequence::FindMetaData<UTakeRecorderSources>
  */
-UCLASS(BlueprintType, Blueprintable)
-class TAKESCORE_API UTakeRecorderSources : public UObject
+UCLASS(MinimalAPI, BlueprintType, Blueprintable)
+class UTakeRecorderSources : public UObject
 {
 public:
 	GENERATED_BODY()
 
-	UTakeRecorderSources(const FObjectInitializer& ObjInit);
+	UE_API UTakeRecorderSources(const FObjectInitializer& ObjInit);
 
 	/**
 	 * Add a new source to this source list of the templated type
@@ -58,7 +60,7 @@ public:
 	 * @return An instance of the specified source type
 	 */
 	UFUNCTION(BlueprintCallable, Category = "Take Recorder", meta = (DeterminesOutputType = "InSourceType"))
-	UTakeRecorderSource* AddSource(TSubclassOf<UTakeRecorderSource> InSourceType);
+	UE_API UTakeRecorderSource* AddSource(TSubclassOf<UTakeRecorderSource> InSourceType);
 
 	/**
 	 * Remove the specified source from this list
@@ -66,7 +68,7 @@ public:
 	 * @param InSource        The source to remove
 	 */
 	UFUNCTION(BlueprintCallable, Category = "Take Recorder")
-	void RemoveSource(UTakeRecorderSource* InSource);
+	UE_API void RemoveSource(UTakeRecorderSource* InSource);
 
 
 	/**
@@ -102,7 +104,7 @@ public:
 
 	/** Calls the recording initialization flows on each of the specified sources. */
 	UFUNCTION(BlueprintCallable, Category = "Take Recorder")
-	void StartRecordingSource(TArray<UTakeRecorderSource*> InSources, const FQualifiedFrameTime& CurrentFrameTime);
+	UE_API void StartRecordingSource(TArray<UTakeRecorderSource*> InSources, const FQualifiedFrameTime& CurrentFrameTime);
 	
 public:
 
@@ -112,61 +114,61 @@ public:
 	 * @param Handler         The delegate to call when the list changes
 	 * @return A handle to this specific binding that should be passed to UnbindSourcesChanged
 	 */
-	FDelegateHandle BindSourcesChanged(const FSimpleDelegate& Handler);
+	UE_API FDelegateHandle BindSourcesChanged(const FSimpleDelegate& Handler);
 
 	/**
 	 * Unbind a previously bound handler for when this source list changes
 	 *
 	 * @param Handle          The handle returned from BindSourcesChanged for the handler to remove
 	 */
-	void UnbindSourcesChanged(FDelegateHandle Handle);
+	UE_API void UnbindSourcesChanged(FDelegateHandle Handle);
 	
 	DECLARE_MULTICAST_DELEGATE_OneParam(FOnSourceAdded, UTakeRecorderSource* Source);
 	DECLARE_MULTICAST_DELEGATE_OneParam(FOnSourceRemoved, UTakeRecorderSource* Source);
 	
 	/** Multicast delegate for when any source is added. */
-	static FOnSourceAdded& OnSourceAdded();
+	static UE_API FOnSourceAdded& OnSourceAdded();
 
 	/** Multicast delegate for when any source is removed. */
-	static FOnSourceRemoved& OnSourceRemoved();
+	static UE_API FOnSourceRemoved& OnSourceRemoved();
 	
 public:
 
 	/*
 	 * Cache assets needed for sequencer.
 	 */
-	void SetCachedAssets(class ULevelSequence* InSequence, FManifestSerializer* InManifestSerializer);
+	UE_API void SetCachedAssets(class ULevelSequence* InSequence, FManifestSerializer* InManifestSerializer);
 
 	/*
 	 * Pre recording pass
 	*
 	*/
-	void PreRecording(class ULevelSequence* InSequence, const FQualifiedFrameTime& InCurrentFrameTime, FManifestSerializer* InManifestSerializer);
+	UE_API void PreRecording(class ULevelSequence* InSequence, const FQualifiedFrameTime& InCurrentFrameTime, FManifestSerializer* InManifestSerializer);
 
 	/*
 	 * Start recording pass
 	 *
 	 */
-	void StartRecording(class ULevelSequence* InSequence, const FQualifiedFrameTime& InCurrentFrameTime, FManifestSerializer* InManifestSerializer);
+	UE_API void StartRecording(class ULevelSequence* InSequence, const FQualifiedFrameTime& InCurrentFrameTime, FManifestSerializer* InManifestSerializer);
 
 	/*
 	 * Moves time forward by given DeltaTime
 	 *
 	 * @return Current Frame Number
 	 */
-	FFrameTime AdvanceTime(const FQualifiedFrameTime& CurrentFrameTime, float DeltaTime);
+	UE_API FFrameTime AdvanceTime(const FQualifiedFrameTime& CurrentFrameTime, float DeltaTime);
 
 	/*
 	* Tick recording pass
 	* @return Current Frame Number we are recording at.
 	*/
-	FFrameTime TickRecording(class ULevelSequence* InSequence, const FQualifiedFrameTime& CurrentFrameTime, float DeltaTime);
+	UE_API FFrameTime TickRecording(class ULevelSequence* InSequence, const FQualifiedFrameTime& CurrentFrameTime, float DeltaTime);
 
 	/*
 	* Stop recording pass
 	*
 	*/
-	void StopRecording(class ULevelSequence* InSequence, const bool bCancelled);
+	UE_API void StopRecording(class ULevelSequence* InSequence, const bool bCancelled);
 
 public:
 	/*
@@ -174,13 +176,13 @@ public:
 	*/
 
 	/** Creates a sub-sequence asset for the specified sub sequence name based on the given root sequence. */
-	static ULevelSequence* CreateSubSequenceForSource(ULevelSequence* InRootSequence, const FString& SubSequenceTrackName,
+	static UE_API ULevelSequence* CreateSubSequenceForSource(ULevelSequence* InRootSequence, const FString& SubSequenceTrackName,
 		const FString& SubSequenceAssetName, const FTakeRecorderSourcesSettings& InSettings = FTakeRecorderSourcesSettings());
 
 	/**
 	 * Array of pairs - key time and the corresponding timecode
 	 */
-	static TArray<TPair<FQualifiedFrameTime, FQualifiedFrameTime> > RecordedTimes;
+	static UE_API TArray<TPair<FQualifiedFrameTime, FQualifiedFrameTime> > RecordedTimes;
 
 	FQualifiedFrameTime GetCachedFrameTime() const { return CachedFrameTime; }
 
@@ -189,28 +191,28 @@ private:
 	virtual void Tick(float DeltaTime) {}
 
 	/** Called whenever a property is changed on this class */
-	virtual void PostEditChangeProperty(FPropertyChangedEvent& PropertyChangedEvent) override;
+	UE_API virtual void PostEditChangeProperty(FPropertyChangedEvent& PropertyChangedEvent) override;
 
 	/** A list of handlers to invoke when the sources list changes */
 	DECLARE_MULTICAST_DELEGATE(FOnSourcesChanged);
 	FOnSourcesChanged OnSourcesChangedEvent;
 
 	/** Calls PreRecording on sources recursively allowing them to create other sources which properly get PreRecording called on them as well. */
-	void PreRecordingRecursive(TArray<UTakeRecorderSource*> InSources, ULevelSequence* InRootSequence, TArray<UTakeRecorderSource*>& NewSourcesOut, FManifestSerializer* InManifestSerializer);
+	UE_API void PreRecordingRecursive(TArray<UTakeRecorderSource*> InSources, ULevelSequence* InRootSequence, TArray<UTakeRecorderSource*>& NewSourcesOut, FManifestSerializer* InManifestSerializer);
 
 	/** Finds the folder that the given Source should be created in, creating it if necessary. */
-	class UMovieSceneFolder* AddFolderForSource(const UTakeRecorderSource* InSource, class UMovieScene* InMovieScene);
+	UE_API class UMovieSceneFolder* AddFolderForSource(const UTakeRecorderSource* InSource, class UMovieScene* InMovieScene);
 
 	/** Remove object bindings that don't have any tracks and are not bindings for attach/path tracks */
-	void RemoveRedundantTracks();
+	UE_API void RemoveRedundantTracks();
 
-	void StartRecordingPreRecordedSources(const FQualifiedFrameTime& CurrentFrameTime);
+	UE_API void StartRecordingPreRecordedSources(const FQualifiedFrameTime& CurrentFrameTime);
 
-	void PreRecordSources(TArray<UTakeRecorderSource *> InSources);
+	UE_API void PreRecordSources(TArray<UTakeRecorderSource *> InSources);
 
-	void StartRecordingTheseSources(const TArray<UTakeRecorderSource *>& InSources, const FQualifiedFrameTime& CurrentFrameTime);
+	UE_API void StartRecordingTheseSources(const TArray<UTakeRecorderSource *>& InSources, const FQualifiedFrameTime& CurrentFrameTime);
 
-	void SetSectionStartTimecode(UMovieSceneSubSection* SubSection, const FTimecode& Timecode, const FQualifiedFrameTime& CurrentFrameTime);
+	UE_API void SetSectionStartTimecode(UMovieSceneSubSection* SubSection, const FTimecode& Timecode, const FQualifiedFrameTime& CurrentFrameTime);
 
 private:
 
@@ -256,3 +258,5 @@ private:
 	/** The last frame time during tick recording */
 	FQualifiedFrameTime CachedFrameTime;
 };
+
+#undef UE_API
diff --git a/Engine/Plugins/VirtualProduction/Takes/Source/TakesCore/Public/TakesCoreBlueprintLibrary.h b/Engine/Plugins/VirtualProduction/Takes/Source/TakesCore/Public/TakesCoreBlueprintLibrary.h
--- a/Engine/Plugins/VirtualProduction/Takes/Source/TakesCore/Public/TakesCoreBlueprintLibrary.h
+++ b/Engine/Plugins/VirtualProduction/Takes/Source/TakesCore/Public/TakesCoreBlueprintLibrary.h
@@ -6,9 +6,11 @@
 #include "AssetRegistry/AssetData.h"
 #include "TakesCoreBlueprintLibrary.generated.h"
 
+#define UE_API TAKESCORE_API
 
-UCLASS()
-class TAKESCORE_API UTakesCoreBlueprintLibrary : public UBlueprintFunctionLibrary
+
+UCLASS(MinimalAPI)
+class UTakesCoreBlueprintLibrary : public UBlueprintFunctionLibrary
 {
 public:
 
@@ -18,7 +20,7 @@ public:
 	 * Compute the next unused sequential take number for the specified slate
 	 */
 	UFUNCTION(BlueprintCallable, Category="Take Recorder")
-	static int32 ComputeNextTakeNumber(const FString& Slate);
+	static UE_API int32 ComputeNextTakeNumber(const FString& Slate);
 
 
 	/**
@@ -28,7 +30,7 @@ public:
 	 * @param TakeNumber   The take number to filter by. <=0 denotes all takes
 	 */
 	UFUNCTION(BlueprintCallable, Category="Take Recorder")
-	static TArray<FAssetData> FindTakes(const FString& Slate, int32 TakeNumber = 0);
+	static UE_API TArray<FAssetData> FindTakes(const FString& Slate, int32 TakeNumber = 0);
 
 public:
 
@@ -37,13 +39,15 @@ public:
 
 	/** Called when the slate is changed. */
 	UFUNCTION(BlueprintCallable, Category = "Take Recorder")
-	static void SetOnTakeRecorderSlateChanged(FOnTakeRecorderSlateChanged OnTakeRecorderSlateChanged);
+	static UE_API void SetOnTakeRecorderSlateChanged(FOnTakeRecorderSlateChanged OnTakeRecorderSlateChanged);
 
 	/** Called when the take number is changed. */
 	UFUNCTION(BlueprintCallable, Category = "Take Recorder")
-	static void SetOnTakeRecorderTakeNumberChanged(FOnTakeRecorderTakeNumberChanged OnTakeRecorderTakeNumberChanged);
+	static UE_API void SetOnTakeRecorderTakeNumberChanged(FOnTakeRecorderTakeNumberChanged OnTakeRecorderTakeNumberChanged);
 
-	static void OnTakeRecorderSlateChanged(const FString& InSlate);
-	static void OnTakeRecorderTakeNumberChanged(int32 InTakeNumber);
+	static UE_API void OnTakeRecorderSlateChanged(const FString& InSlate);
+	static UE_API void OnTakeRecorderTakeNumberChanged(int32 InTakeNumber);
 
 };
+
+#undef UE_API
diff --git a/Engine/Plugins/XGEController/Source/Public/XGEControllerModule.h b/Engine/Plugins/XGEController/Source/Public/XGEControllerModule.h
--- a/Engine/Plugins/XGEController/Source/Public/XGEControllerModule.h
+++ b/Engine/Plugins/XGEController/Source/Public/XGEControllerModule.h
@@ -69,7 +69,7 @@ public:
 	virtual void InitializeController() override final;
 	virtual bool SupportsLocalWorkers() override final;
 
-	XGECONTROLLER_API virtual const FString GetName() override final { return FString("XGE Controller"); };
+	virtual const FString GetName() override final { return FString("XGE Controller"); };
 
 	XGECONTROLLER_API virtual bool IsSupported() override final;
 
