Take Recorder: Clean up TakeRecorderModule:

- Sort function definitions by declaration order
- Group function declarations by interface
- Add Begin / End Interface comments

- Remove unused variable TakePresetActions
- Remove ProjectSettingsName (adds noise to reader)
- Group similar functions (removed some carriage returns)

diff --git a/Engine/Plugins/VirtualProduction/Takes/Source/TakeRecorder/Private/TakeRecorderModule.cpp b/Engine/Plugins/VirtualProduction/Takes/Source/TakeRecorder/Private/TakeRecorderModule.cpp
--- a/Engine/Plugins/VirtualProduction/Takes/Source/TakeRecorder/Private/TakeRecorderModule.cpp
+++ b/Engine/Plugins/VirtualProduction/Takes/Source/TakeRecorder/Private/TakeRecorderModule.cpp
@@ -54,6 +54,7 @@
 #include "MovieSceneSequence.h"
 #include "Sections/MovieSceneSubSection.h"
 #include "MovieScene.h"
+#include "UObject/Class.h"
 
 #include "Algo/RemoveIf.h"
 #include "Engine/Font.h"
@@ -319,103 +320,11 @@
 	}
 }
 
-void FTakeRecorderModule::RegisterMenus()
-{
-#if WITH_EDITOR
-	if (!UToolMenus::IsToolMenuUIEnabled())
-	{
-		return;
-	}
-
-	FToolMenuOwnerScoped MenuOwner("TakeRecorder");
-	UToolMenus* ToolMenus = UToolMenus::Get();
-	UToolMenu* Menu = ToolMenus->ExtendMenu("ContentBrowser.AssetContextMenu.LevelSequence");
-	if (!Menu)
-	{
-		return;
-	}
-
-	FToolMenuSection& Section = Menu->FindOrAddSection("GetAssetActions");
-	Section.AddDynamicEntry("TakeRecorderActions", FNewToolMenuSectionDelegate::CreateLambda([](FToolMenuSection& InSection)
-	{
-		UContentBrowserAssetContextMenuContext* Context = InSection.FindContext<UContentBrowserAssetContextMenuContext>();
-		if (!Context)
-		{
-			return;
-		}
-
-		if (Context->SelectedAssets.Num() == 1 && Context->SelectedAssets[0].IsInstanceOf(ULevelSequence::StaticClass()))
-		{
-			const FAssetData LevelSequenceAsset = Context->SelectedAssets[0];
-
-			InSection.AddMenuEntry(
-				"OpenInTakeRecorder_Label",
-				LOCTEXT("OpenInTakeRecorder_Label", "Open in Take Recorder"),
-				LOCTEXT("OpenInTakeRecorder_Tooltip", "Opens this level sequence asset in Take Recorder by copying its contents into the pending take"),
-				FSlateIcon(FTakeRecorderStyle::StyleName, "TakeRecorder.TabIcon"),
-				FExecuteAction::CreateLambda(
-					[LevelSequenceAsset]
-					{
-						if (ULevelSequence* LevelSequence = Cast<ULevelSequence>(LevelSequenceAsset.GetAsset()))
-						{
-							FLevelEditorModule& LevelEditorModule = FModuleManager::LoadModuleChecked<FLevelEditorModule>("LevelEditor");
-							TSharedPtr<SDockTab> DockTab = LevelEditorModule.GetLevelEditorTabManager()->TryInvokeTab(ITakeRecorderModule::TakeRecorderTabName);
-							if (DockTab.IsValid())
-							{
-								TSharedRef<STakeRecorderTabContent> TabContent = StaticCastSharedRef<STakeRecorderTabContent>(DockTab->GetContent());
-
-								// If this sequence has already been recorded, set it up for viewing, otherwise start recording from it.
-								UTakeMetaData* TakeMetaData = LevelSequence->FindMetaData<UTakeMetaData>();
-								if (!TakeMetaData || !TakeMetaData->Recorded())
-								{
-									TabContent->SetupForRecording(LevelSequence);
-								}
-								else
-								{
-									TabContent->SetupForViewing(LevelSequence);
-								}
-							}
-						}
-					}
-				)
-			);
-			InSection.AddMenuEntry(
-				"RecordIntoTakeRecorder_Label",
-				LOCTEXT("RecordWithTakeRecorder_Label", "Record with Take Recorder"),
-				LOCTEXT("RecordWithTakeRecorder_Tooltip", "Opens this level sequence asset for recording into with Take Recorder"),
-				FSlateIcon(FTakeRecorderStyle::StyleName, "TakeRecorder.TabIcon"),
-				FExecuteAction::CreateLambda(
-					[LevelSequenceAsset]
-					{
-						if (ULevelSequence* LevelSequence = Cast<ULevelSequence>(LevelSequenceAsset.GetAsset()))
-						{
-							FLevelEditorModule& LevelEditorModule = FModuleManager::LoadModuleChecked<FLevelEditorModule>("LevelEditor");
-							TSharedPtr<SDockTab> DockTab = LevelEditorModule.GetLevelEditorTabManager()->TryInvokeTab(ITakeRecorderModule::TakeRecorderTabName);
-							if (DockTab.IsValid())
-							{
-								TSharedRef<STakeRecorderTabContent> TabContent = StaticCastSharedRef<STakeRecorderTabContent>(DockTab->GetContent());
-
-								TabContent->SetupForRecordingInto(LevelSequence);
-							}
-						}
-					}
-				)
-			);
-		}
-	}));
-#endif // WITH_EDITOR
-}
-
 FTakeRecorderModule::FTakeRecorderModule()
 	: SequencerSettings(nullptr)
 {
 }
 
-UTakePreset* FTakeRecorderModule::GetPendingTake() const
-{
-	return FindObject<UTakePreset>(nullptr, TEXT("/Temp/TakeRecorder/PendingTake.PendingTake"));
-}
-
 void FTakeRecorderModule::StartupModule()
 {
 	FTakeRecorderStyle::Get();
@@ -447,14 +356,6 @@
 #endif
 }
 
-void FTakeRecorderModule::OnEditorClose()
-{
-	if (UTakeRecorder* ActiveRecorder = UTakeRecorder::GetActiveRecorder())
-	{
-		ActiveRecorder->Stop();
-	}
-}
-
 void FTakeRecorderModule::ShutdownModule()
 {
 #if WITH_EDITOR
@@ -474,6 +375,11 @@
 	UnregisterSerializedRecorder();
 }
 
+UTakePreset* FTakeRecorderModule::GetPendingTake() const
+{
+	return FindObject<UTakePreset>(nullptr, TEXT("/Temp/TakeRecorder/PendingTake.PendingTake"));
+}
+
 void FTakeRecorderModule::RegisterExternalObject(UObject* InExternalObject)
 {
 	ExternalObjects.Add(InExternalObject);
@@ -486,30 +392,6 @@
 	ExternalObjects.Remove(InExternalObject);
 }
 
-void FTakeRecorderModule::RegisterDetailCustomizations()
-{
-#if WITH_EDITOR
-
-	if (GIsEditor)
-	{
-		FPropertyEditorModule& PropertyEditorModule = FModuleManager::Get().LoadModuleChecked<FPropertyEditorModule>("PropertyEditor");
-		ProjectSettingsName = UTakeRecorderProjectSettings::StaticClass()->GetFName();
-
-		PropertyEditorModule.RegisterCustomClassLayout(ProjectSettingsName, FOnGetDetailCustomizationInstance::CreateLambda(&MakeShared<FTakeRecorderProjectSettingsCustomization>));
-	}
-
-#endif
-}
-
-void FTakeRecorderModule::UnregisterDetailCustomizations()
-{
-	FPropertyEditorModule* PropertyEditorModule = FModuleManager::Get().GetModulePtr<FPropertyEditorModule>("PropertyEditor");
-	if (PropertyEditorModule)
-	{
-		PropertyEditorModule->UnregisterCustomClassLayout(ProjectSettingsName);
-	}
-}
-
 FDelegateHandle FTakeRecorderModule::RegisterSourcesMenuExtension(const FOnExtendSourcesMenu& InExtension)
 {
 	return SourcesMenuExtenderEvent.Add(InExtension);
@@ -535,6 +417,36 @@
 	GetMutableDefault<UTakeRecorderProjectSettings>()->AdditionalSettings.Add(InSettingsObject);
 }
 
+void FTakeRecorderModule::AddReferencedObjects(FReferenceCollector& Collector)
+{
+	if (SequencerSettings)
+	{
+		Collector.AddReferencedObject(SequencerSettings);
+	}
+}
+
+void FTakeRecorderModule::RegisterDetailCustomizations()
+{
+#if WITH_EDITOR
+	if (GIsEditor)
+	{
+		FPropertyEditorModule& PropertyEditorModule = FModuleManager::Get().LoadModuleChecked<FPropertyEditorModule>("PropertyEditor");
+		PropertyEditorModule.RegisterCustomClassLayout(
+			UTakeRecorderProjectSettings::StaticClass()->GetFName(),
+			FOnGetDetailCustomizationInstance::CreateLambda(&MakeShared<FTakeRecorderProjectSettingsCustomization>)
+			);
+	}
+#endif
+}
+
+void FTakeRecorderModule::UnregisterDetailCustomizations()
+{
+	if (FPropertyEditorModule* PropertyEditorModule = FModuleManager::Get().GetModulePtr<FPropertyEditorModule>("PropertyEditor"))
+	{
+		PropertyEditorModule->UnregisterCustomClassLayout(UTakeRecorderProjectSettings::StaticClass()->GetFName());
+	}
+}
+
 void FTakeRecorderModule::RegisterLevelEditorExtensions()
 {
 #if WITH_EDITOR
@@ -563,7 +475,7 @@
 #endif
 }
 
-void FTakeRecorderModule::UnregisterLevelEditorExtensions()
+void FTakeRecorderModule::UnregisterLevelEditorExtensions() const
 {
 #if WITH_EDITOR
 	if (FSlateApplication::IsInitialized())
@@ -626,7 +538,7 @@
 	}
 }
 
-void FTakeRecorderModule::PopulateSourcesMenu(TSharedRef<FExtender> InExtender, UTakeRecorderSources* InSources)
+void FTakeRecorderModule::PopulateSourcesMenu(TSharedRef<FExtender> InExtender, UTakeRecorderSources* InSources) const
 {
 	SourcesMenuExtenderEvent.Broadcast(InExtender, InSources);
 }
@@ -637,16 +549,103 @@
 	IModularFeatures::Get().RegisterModularFeature(FSerializedRecorder::ModularFeatureName, SerializedRecorder.Get());
 }
 
-void FTakeRecorderModule::UnregisterSerializedRecorder()
+void FTakeRecorderModule::UnregisterSerializedRecorder() const
 {
 	IModularFeatures::Get().UnregisterModularFeature(FSerializedRecorder::ModularFeatureName, SerializedRecorder.Get());
 }
 
-void FTakeRecorderModule::AddReferencedObjects(FReferenceCollector& Collector)
+void FTakeRecorderModule::RegisterMenus()
 {
-	if (SequencerSettings)
+#if WITH_EDITOR
+	if (!UToolMenus::IsToolMenuUIEnabled())
 	{
-		Collector.AddReferencedObject(SequencerSettings);
+		return;
+	}
+
+	FToolMenuOwnerScoped MenuOwner("TakeRecorder");
+	UToolMenus* ToolMenus = UToolMenus::Get();
+	UToolMenu* Menu = ToolMenus->ExtendMenu("ContentBrowser.AssetContextMenu.LevelSequence");
+	if (!Menu)
+	{
+		return;
+	}
+
+	FToolMenuSection& Section = Menu->FindOrAddSection("GetAssetActions");
+	Section.AddDynamicEntry("TakeRecorderActions", FNewToolMenuSectionDelegate::CreateLambda([](FToolMenuSection& InSection)
+	{
+		UContentBrowserAssetContextMenuContext* Context = InSection.FindContext<UContentBrowserAssetContextMenuContext>();
+		if (!Context)
+		{
+			return;
+		}
+
+		if (Context->SelectedAssets.Num() == 1 && Context->SelectedAssets[0].IsInstanceOf(ULevelSequence::StaticClass()))
+		{
+			const FAssetData LevelSequenceAsset = Context->SelectedAssets[0];
+
+			InSection.AddMenuEntry(
+				"OpenInTakeRecorder_Label",
+				LOCTEXT("OpenInTakeRecorder_Label", "Open in Take Recorder"),
+				LOCTEXT("OpenInTakeRecorder_Tooltip", "Opens this level sequence asset in Take Recorder by copying its contents into the pending take"),
+				FSlateIcon(FTakeRecorderStyle::StyleName, "TakeRecorder.TabIcon"),
+				FExecuteAction::CreateLambda(
+					[LevelSequenceAsset]
+					{
+						if (ULevelSequence* LevelSequence = Cast<ULevelSequence>(LevelSequenceAsset.GetAsset()))
+						{
+							FLevelEditorModule& LevelEditorModule = FModuleManager::LoadModuleChecked<FLevelEditorModule>("LevelEditor");
+							TSharedPtr<SDockTab> DockTab = LevelEditorModule.GetLevelEditorTabManager()->TryInvokeTab(ITakeRecorderModule::TakeRecorderTabName);
+							if (DockTab.IsValid())
+							{
+								TSharedRef<STakeRecorderTabContent> TabContent = StaticCastSharedRef<STakeRecorderTabContent>(DockTab->GetContent());
+
+								// If this sequence has already been recorded, set it up for viewing, otherwise start recording from it.
+								UTakeMetaData* TakeMetaData = LevelSequence->FindMetaData<UTakeMetaData>();
+								if (!TakeMetaData || !TakeMetaData->Recorded())
+								{
+									TabContent->SetupForRecording(LevelSequence);
+								}
+								else
+								{
+									TabContent->SetupForViewing(LevelSequence);
+								}
+							}
+						}
+					}
+				)
+			);
+			InSection.AddMenuEntry(
+				"RecordIntoTakeRecorder_Label",
+				LOCTEXT("RecordWithTakeRecorder_Label", "Record with Take Recorder"),
+				LOCTEXT("RecordWithTakeRecorder_Tooltip", "Opens this level sequence asset for recording into with Take Recorder"),
+				FSlateIcon(FTakeRecorderStyle::StyleName, "TakeRecorder.TabIcon"),
+				FExecuteAction::CreateLambda(
+					[LevelSequenceAsset]
+					{
+						if (ULevelSequence* LevelSequence = Cast<ULevelSequence>(LevelSequenceAsset.GetAsset()))
+						{
+							FLevelEditorModule& LevelEditorModule = FModuleManager::LoadModuleChecked<FLevelEditorModule>("LevelEditor");
+							TSharedPtr<SDockTab> DockTab = LevelEditorModule.GetLevelEditorTabManager()->TryInvokeTab(ITakeRecorderModule::TakeRecorderTabName);
+							if (DockTab.IsValid())
+							{
+								TSharedRef<STakeRecorderTabContent> TabContent = StaticCastSharedRef<STakeRecorderTabContent>(DockTab->GetContent());
+
+								TabContent->SetupForRecordingInto(LevelSequence);
+							}
+						}
+					}
+				)
+			);
+		}
+	}));
+#endif // WITH_EDITOR
+}
+
+void FTakeRecorderModule::OnEditorClose()
+{
+	if (UTakeRecorder* ActiveRecorder = UTakeRecorder::GetActiveRecorder())
+	{
+		ActiveRecorder->Stop();
 	}
 }
 
diff --git a/Engine/Plugins/VirtualProduction/Takes/Source/TakeRecorder/Private/TakeRecorderModule.h b/Engine/Plugins/VirtualProduction/Takes/Source/TakeRecorder/Private/TakeRecorderModule.h
--- a/Engine/Plugins/VirtualProduction/Takes/Source/TakeRecorder/Private/TakeRecorderModule.h
+++ b/Engine/Plugins/VirtualProduction/Takes/Source/TakeRecorder/Private/TakeRecorderModule.h
@@ -7,6 +7,7 @@
 #include "UObject/GCObject.h"
 #include "ITakeRecorderModule.h"
 #include "UObject/WeakObjectPtrTemplates.h"
+#include "Templates/UniquePtr.h"
 #include "Widgets/Input/SCheckBox.h"
 
 class FExtender;
@@ -19,72 +20,55 @@ class USequencerSettings;
 class FTakeRecorderModule : public ITakeRecorderModule, public FGCObject
 {
 public:
+	
 	FTakeRecorderModule();
 
-	void PopulateSourcesMenu(TSharedRef<FExtender> InExtender, UTakeRecorderSources* InSources);
-	virtual UTakePreset* GetPendingTake() const override;
+	void PopulateSourcesMenu(TSharedRef<FExtender> InExtender, UTakeRecorderSources* InSources) const;
 
-	//~ ITakeRecorderModule API
-	FOnGenerateWidgetExtensions& GetToolbarExtensionGenerators() override { return ToolbarExtensionGenerators; }
-	FOnGenerateWidgetExtensions& GetRecordButtonExtensionGenerators() override { return ButtonExtensionGenerators; }
-	FOnExternalObjectAddRemoveEvent& GetExternalObjectAddRemoveEventDelegate() override { return ExternalObjectAddRemoveEvent; }
-	FOnRecordErrorCheck&  GetRecordErrorCheckGenerator() override { return RecordErrorCheck; }
-	TArray<TWeakObjectPtr<>>& GetExternalObjects() override { return ExternalObjects; }
-	FLastRecordedLevelSequenceProvider& GetLastLevelSequenceProvider() override { return LastLevelSequenceProvider; }
-	FCanReviewLastRecordedLevelSequence& GetCanReviewLastRecordedLevelSequenceDelegate() override { return CanReviewLastRecordedSequence; };
-
-	void RegisterExternalObject(UObject* InExternalObject) override;
-	void UnregisterExternalObject(UObject* InExternalObject) override;
-
-	FOnForceSaveAsPreset& OnForceSaveAsPreset() override
-	{
-		return ForceSaveAsPresetEvent;
-	}
-
-	FSourceExtensionData& GetSourcesExtensionData()
-	{
-		return SourceExtensionData;
-	}
+	//~ Begin IModuleInterface Interface
+	virtual void StartupModule() override;
+	virtual void ShutdownModule() override;
+	//~ End IModuleInterface Interface
 	
-private:
-
-	void StartupModule() override;
-	void ShutdownModule() override;
-
-	FDelegateHandle RegisterSourcesMenuExtension(const FOnExtendSourcesMenu& InExtension) override;
+	//~ Begin ITakeRecorderModule Interface
+	virtual UTakePreset* GetPendingTake() const override;
+	virtual FOnGenerateWidgetExtensions& GetToolbarExtensionGenerators() override { return ToolbarExtensionGenerators; }
+	virtual FOnGenerateWidgetExtensions& GetRecordButtonExtensionGenerators() override { return ButtonExtensionGenerators; }
+	virtual FOnExternalObjectAddRemoveEvent& GetExternalObjectAddRemoveEventDelegate() override { return ExternalObjectAddRemoveEvent; }
+	virtual FOnRecordErrorCheck&  GetRecordErrorCheckGenerator() override { return RecordErrorCheck; }
+	virtual TArray<TWeakObjectPtr<>>& GetExternalObjects() override { return ExternalObjects; }
+	virtual FLastRecordedLevelSequenceProvider& GetLastLevelSequenceProvider() override { return LastLevelSequenceProvider; }
+	virtual FCanReviewLastRecordedLevelSequence& GetCanReviewLastRecordedLevelSequenceDelegate() override { return CanReviewLastRecordedSequence; };
+	virtual void RegisterExternalObject(UObject* InExternalObject) override;
+	virtual void UnregisterExternalObject(UObject* InExternalObject) override;
+	virtual FOnForceSaveAsPreset& OnForceSaveAsPreset() override { return ForceSaveAsPresetEvent; }
+	virtual FDelegateHandle RegisterSourcesMenuExtension(const FOnExtendSourcesMenu& InExtension) override;
 	virtual void RegisterSourcesExtension(const FSourceExtensionData& InData) override;
 	virtual void UnregisterSourcesExtension() override;
-	void UnregisterSourcesMenuExtension(FDelegateHandle Handle) override;
-	void RegisterSettingsObject(UObject* InSettingsObject) override;
-
-	/** FGCObject interface */
-	void AddReferencedObjects(FReferenceCollector& Collector) override;
-	FString GetReferencerName() const override
-	{
-		return "FTakeRecorderModule";
-	}
-
+	virtual void UnregisterSourcesMenuExtension(FDelegateHandle Handle) override;
+	virtual void RegisterSettingsObject(UObject* InSettingsObject) override;
+	//~ Begin ITakeRecorderModule Interface
+	
+	FSourceExtensionData& GetSourcesExtensionData() { return SourceExtensionData; }
+	
 private:
+	
+	//~ Begin FGCObject Interface
+	virtual void AddReferencedObjects(FReferenceCollector& Collector) override;
+	virtual FString GetReferencerName() const override { return "FTakeRecorderModule"; }
+	//~ End FGCObject Interface
 
 	void RegisterDetailCustomizations();
-
 	void UnregisterDetailCustomizations();
 
 	void RegisterLevelEditorExtensions();
-
-	void UnregisterLevelEditorExtensions();
-
-	void RegisterAssetTools();
-
-	void UnregisterAssetTools();
+	void UnregisterLevelEditorExtensions() const;
 
 	void RegisterSettings();
-
 	void UnregisterSettings();
 
 	void RegisterSerializedRecorder();
-
-	void UnregisterSerializedRecorder();
+	void UnregisterSerializedRecorder() const;
 
 	void RegisterMenus();
 
@@ -108,14 +92,12 @@ private:
 	FDelegateHandle LevelEditorTabManagerChangedHandle;
 	FDelegateHandle ModulesChangedHandle;
 
-	/** Cached name of the project settings fo de-registration of details customizations on shutdown (after UObject destruction) */
-	FName ProjectSettingsName;
-
-	TSharedPtr<FTakePresetActions> TakePresetActions;
 	TSharedPtr<FSerializedRecorder> SerializedRecorder;
 
+	/** Extra class default objects that should appear in the Take Recorder panel. */
 	TArray<TWeakObjectPtr<>> ExternalObjects;
 
+	/** The settings object registers with the Settings module. */
 	TObjectPtr<USequencerSettings> SequencerSettings;
 
 	/**
@@ -123,4 +105,4 @@ private:
 	 * as only the source module should be utilizing this.
 	 */
 	FSourceExtensionData SourceExtensionData;
-};
+};
\ No newline at end of file
