TakeRecorder: Fix not applying time dilation when sequencer was closed.

diff --git a/Engine/Plugins/VirtualProduction/Takes/Source/TakeRecorder/Private/Recorder/TakeRecorder.cpp b/Engine/Plugins/VirtualProduction/Takes/Source/TakeRecorder/Private/Recorder/TakeRecorder.cpp
--- a/Engine/Plugins/VirtualProduction/Takes/Source/TakeRecorder/Private/Recorder/TakeRecorder.cpp
+++ b/Engine/Plugins/VirtualProduction/Takes/Source/TakeRecorder/Private/Recorder/TakeRecorder.cpp
@@ -798,7 +798,24 @@ void UTakeRecorder::InternalTick(float DeltaTime)
 
 	TSharedPtr<ISequencer> Sequencer = WeakSequencer.Pin();
 	FQualifiedFrameTime RecordTime = GetRecordTime();
-		
+
+	if (!Sequencer.IsValid())
+	{
+		const UWorld* RecordingWorld = GetWorld();
+		check(RecordingWorld);
+		if (const AWorldSettings* WorldSettings = RecordingWorld->GetWorldSettings())
+		{
+			// If we're not using Sequencer we have to apply time dilation ourselves as this normally happens in Sequencer.cpp
+			// when the TimeController is ticked.
+			FQualifiedFrameTime ScaledRecordTime(
+				RecordTime.Time * WorldSettings->CinematicTimeDilation,
+				RecordTime.Rate
+			);
+
+			RecordTime = MoveTemp(ScaledRecordTime);
+		}
+	}
+	
 	UTakeRecorderSources* Sources = SequenceAsset->FindOrAddMetaData<UTakeRecorderSources>();
 	if (!Parameters.bDisableRecordingAndSave)
 	{
