Take Recorder: Introduce TakeRecorder.EnableHitchlessTimecode, which is disabled by default. When enabled, take recorder will wrap the current timecode with UTimecodeRegressionProvider, which estimates the timecode the current frame should have. When combined with a fixed custom engine timestep that implements IClockedTimeStep, long engine frames will not cause subsequent frames to skip timecode. Before, subsequent frames would skip timecode proportional to the delay of the engine hitch length.

This shows a provisonal dialogue box to help the user set up fixed engine timestep is not set up, or a warning notififcation if an invalid one is set up. This UX is subject to change and intended for internal testing. The UX is supposed to be improved in the future!

diff --git a/Engine/Plugins/VirtualProduction/Takes/Source/TakeRecorder/Private/TakeRecorderModule.cpp b/Engine/Plugins/VirtualProduction/Takes/Source/TakeRecorder/Private/TakeRecorderModule.cpp
--- a/Engine/Plugins/VirtualProduction/Takes/Source/TakeRecorder/Private/TakeRecorderModule.cpp
+++ b/Engine/Plugins/VirtualProduction/Takes/Source/TakeRecorder/Private/TakeRecorderModule.cpp
@@ -334,6 +334,7 @@ void FTakeRecorderModule::StartupModule()
 	RegisterLevelEditorExtensions();
 	RegisterSettings();
 	RegisterSerializedRecorder();
+	HitchlessTimecodeIntegration = MakeUnique<UE::TakeRecorder::FHitchlessTimecodeRecordLogic>();
 
 #if WITH_EDITOR
 	if (GIsEditor)
@@ -373,6 +374,7 @@ void FTakeRecorderModule::ShutdownModule()
 	UnregisterLevelEditorExtensions();
 	UnregisterSettings();
 	UnregisterSerializedRecorder();
+	HitchlessTimecodeIntegration.Reset();
 }
 
 UTakePreset* FTakeRecorderModule::GetPendingTake() const
diff --git a/Engine/Plugins/VirtualProduction/Takes/Source/TakeRecorder/Private/TakeRecorderModule.h b/Engine/Plugins/VirtualProduction/Takes/Source/TakeRecorder/Private/TakeRecorderModule.h
--- a/Engine/Plugins/VirtualProduction/Takes/Source/TakeRecorder/Private/TakeRecorderModule.h
+++ b/Engine/Plugins/VirtualProduction/Takes/Source/TakeRecorder/Private/TakeRecorderModule.h
@@ -7,7 +7,7 @@
 #include "UObject/GCObject.h"
 #include "ITakeRecorderModule.h"
 #include "UObject/WeakObjectPtrTemplates.h"
-#include "Templates/UniquePtr.h"
+#include "Timecode/HitchlessTimecodeRecordLogic.h"
 #include "Widgets/Input/SCheckBox.h"
 
 class FExtender;
@@ -20,7 +20,6 @@ class USequencerSettings;
 class FTakeRecorderModule : public ITakeRecorderModule, public FGCObject
 {
 public:
-	
 	FTakeRecorderModule();
 
 	void PopulateSourcesMenu(TSharedRef<FExtender> InExtender, UTakeRecorderSources* InSources) const;
@@ -94,6 +93,15 @@ private:
 
 	TSharedPtr<FSerializedRecorder> SerializedRecorder;
 
+	/**
+	 * Sets up the UEngine::TimecodeProvider so timecode is not skipped when the engine hitches.
+	 * 
+	 * For future maintainers: this could totally live in a separate module / plugin / whatever.
+	 * There's no technical reason to have this live in TakeRecorder. The reason it's here is that we probably always want hitching to be handled
+	 * gracefully.
+	 */
+	TUniquePtr<UE::TakeRecorder::FHitchlessTimecodeRecordLogic> HitchlessTimecodeIntegration;
+
 	/** Extra class default objects that should appear in the Take Recorder panel. */
 	TArray<TWeakObjectPtr<>> ExternalObjects;
 
diff --git a/Engine/Plugins/VirtualProduction/Takes/Source/TakeRecorder/Private/Timecode/GuardTimecodeProvider.h b/Engine/Plugins/VirtualProduction/Takes/Source/TakeRecorder/Private/Timecode/GuardTimecodeProvider.h
new file mode 100644
--- /dev/null
+++ b/Engine/Plugins/VirtualProduction/Takes/Source/TakeRecorder/Private/Timecode/GuardTimecodeProvider.h
@@ -0,0 +1,37 @@
+﻿// Copyright Epic Games, Inc. All Rights Reserved.
+
+#pragma once
+
+#include "Engine/Engine.h"
+#include "Engine/TimecodeProvider.h"
+#include "Templates/UnrealTemplate.h"
+#include "UObject/StrongObjectPtrTemplates.h"
+
+namespace UE::TakeRecorder
+{
+/** Overrides the engine's timecode provider for the lifetime of the instance. */
+class FGuardTimecodeProvider
+	: public FNoncopyable	// We don't expect any copying
+{
+	TStrongObjectPtr<UTimecodeProvider> RestoreProvider = nullptr;
+public:
+
+	explicit FGuardTimecodeProvider(UTimecodeProvider* Override)
+	{
+		if (GEngine)
+		{
+			RestoreProvider.Reset(GEngine->GetTimecodeProvider());
+			GEngine->SetTimecodeProvider(Override);
+		}
+	}
+
+	~FGuardTimecodeProvider()
+	{
+		if (GEngine)
+		{
+			GEngine->SetTimecodeProvider(RestoreProvider.Get());
+		}
+	}
+};
+}
+
diff --git a/Engine/Plugins/VirtualProduction/Takes/Source/TakeRecorder/Private/Timecode/GuardTimestep.h b/Engine/Plugins/VirtualProduction/Takes/Source/TakeRecorder/Private/Timecode/GuardTimestep.h
new file mode 100644
--- /dev/null
+++ b/Engine/Plugins/VirtualProduction/Takes/Source/TakeRecorder/Private/Timecode/GuardTimestep.h
@@ -0,0 +1,37 @@
+﻿// Copyright Epic Games, Inc. All Rights Reserved.
+
+#pragma once
+
+#include "Engine/Engine.h"
+#include "Engine/EngineCustomTimeStep.h"
+#include "Templates/UnrealTemplate.h"
+#include "UObject/StrongObjectPtrTemplates.h"
+
+namespace UE::TakeRecorder
+{
+/** Overrides the engine's timestep for the lifetime of the instance. */
+class FGuardTimestep
+	: public FNoncopyable	// We don't expect any copying
+{
+	TStrongObjectPtr<UEngineCustomTimeStep> RestoreTimestep = nullptr;
+public:
+
+	explicit FGuardTimestep(UEngineCustomTimeStep* Override)
+	{
+		if (GEngine)
+		{
+			RestoreTimestep.Reset(GEngine->GetCustomTimeStep());
+			GEngine->SetCustomTimeStep(Override);
+		}
+	}
+
+	~FGuardTimestep()
+	{
+		if (GEngine)
+		{
+			GEngine->SetCustomTimeStep(RestoreTimestep.Get());
+		}
+	}
+};
+}
+
diff --git a/Engine/Plugins/VirtualProduction/Takes/Source/TakeRecorder/Private/Timecode/HitchlessTimecodeRecordLogic.cpp b/Engine/Plugins/VirtualProduction/Takes/Source/TakeRecorder/Private/Timecode/HitchlessTimecodeRecordLogic.cpp
new file mode 100644
--- /dev/null
+++ b/Engine/Plugins/VirtualProduction/Takes/Source/TakeRecorder/Private/Timecode/HitchlessTimecodeRecordLogic.cpp
@@ -0,0 +1,119 @@
+﻿// Copyright Epic Games, Inc. All Rights Reserved.
+
+#include "HitchlessTimecodeRecordLogic.h"
+
+#include "CatchupFixedRateCustomTimeStep.h"
+#include "ClassViewerFilter.h"
+#include "Engine/Engine.h"
+#include "Engine/EngineCustomTimeStep.h"
+#include "Estimation/TimecodeRegressionProvider.h"
+#include "Framework/Notifications/NotificationManager.h"
+#include "Misc/CoreDelegates.h"
+#include "Recorder/TakeRecorderSubsystem.h"
+#include "TimecodeNotifications.h"
+#include "UObject/Package.h"
+
+namespace UE::TakeRecorder
+{
+namespace Private
+{
+static TAutoConsoleVariable<bool> CVarEnableHitchlessTimecodeProvider(
+	TEXT("TakeRecorder.EnableHitchlessTimecode"),
+	false,
+	TEXT("If true, we'll try to handle engine hitching. This requires that the UEngineCustomTimestep implements the IClockedTimeStep interface.")
+	);
+static TAutoConsoleVariable<float> CVarTimecodeLinearRegressionSeconds(
+	TEXT("TakeRecorder.TimecodeLinearRegressionSeconds"),
+	1.0f,
+	TEXT("When using linear regression to predict timecode, how many seconds worth of frames to sample.")
+	);
+
+/** @return The timecode estimator instance we'll use to estimate timecode - effectively handling any engine hitches. */
+static UTimecodeRegressionProvider* CreateTimecodeEstimator(UTimecodeProvider* InActualImplementation)
+{
+	UTimecodeRegressionProvider* TimecodeRegression = NewObject<UTimecodeRegressionProvider>(GetTransientPackage());
+	TimecodeRegression->SetTimecodeImplementation(InActualImplementation);
+	
+	// How many frames do we need to sample?
+	// E.g. if FrameRate is 60 FPS of 60 frames, then Numerator = 60 meaning 60 / 1 = 60. We want
+	const FFrameRate FrameRate = GEngine->GetEngineSubsystem<UTakeRecorderSubsystem>()->GetFrameRate();
+	TimecodeRegression->NumSampledFrames = CVarTimecodeLinearRegressionSeconds.GetValueOnAnyThread() * FrameRate.Numerator / FrameRate.Denominator;
+	
+	return TimecodeRegression;
+}
+}
+	
+FHitchlessTimecodeRecordLogic::FHitchlessTimecodeRecordLogic()
+{
+	if (IsEngineStartupModuleLoadingComplete())
+	{
+		OnPostEngineInit();
+	}
+	else
+	{
+		FCoreDelegates::OnPostEngineInit.AddRaw(this, &FHitchlessTimecodeRecordLogic::OnPostEngineInit);
+	}
+}
+
+FHitchlessTimecodeRecordLogic::~FHitchlessTimecodeRecordLogic()
+{
+	FCoreDelegates::OnPostEngineInit.RemoveAll(this);
+	
+	CleanUp();
+}
+
+void FHitchlessTimecodeRecordLogic::OnPostEngineInit()
+{
+	UTakeRecorderSubsystem* TakeRecorderSubsystem = GEngine ? GEngine->GetEngineSubsystem<UTakeRecorderSubsystem>() : nullptr;
+	if (ensure(TakeRecorderSubsystem))
+	{
+		TakeRecorderSubsystem->GetOnRecordingInitializedEvent().AddRaw(this, &FHitchlessTimecodeRecordLogic::OnInitializeRecording);
+		TakeRecorderSubsystem->GetOnRecordingFinishedEvent().AddRaw(this, &FHitchlessTimecodeRecordLogic::OnFinishRecording);
+		TakeRecorderSubsystem->GetOnRecordingCancelledEvent().AddRaw(this, &FHitchlessTimecodeRecordLogic::OnCancelRecording);
+	}
+}
+
+void FHitchlessTimecodeRecordLogic::CleanUp()
+{
+	if (UTakeRecorderSubsystem* TakeRecorderSubsystem = GEngine ? GEngine->GetEngineSubsystem<UTakeRecorderSubsystem>() : nullptr)
+	{
+		TakeRecorderSubsystem->GetOnRecordingInitializedEvent().RemoveAll(this);
+		TakeRecorderSubsystem->GetOnRecordingFinishedEvent().RemoveAll(this);
+		TakeRecorderSubsystem->GetOnRecordingCancelledEvent().RemoveAll(this);
+		
+		RecordData.Reset();
+	}
+}
+
+void FHitchlessTimecodeRecordLogic::SetupEngineEnvironmentForRecording()
+{
+	UTimecodeProvider* CurrentTimecodeProvider = GEngine ? GEngine->GetTimecodeProvider() : nullptr;
+	if (!Private::CVarEnableHitchlessTimecodeProvider.GetValueOnAnyThread()
+		|| !CurrentTimecodeProvider)
+	{
+		return;
+	}
+
+	UEngineCustomTimeStep* PreexistingTimestep = GEngine->GetCustomTimeStep();
+	if (PreexistingTimestep && !UTimecodeRegressionProvider::SupportsTimestep(PreexistingTimestep))
+	{
+		ShowIncompatibleTimestepNotification(PreexistingTimestep);
+		return;
+	}
+
+	if (PreexistingTimestep)
+	{
+		// Use pre-configured custom timestep
+		RecordData.Emplace(Private::CreateTimecodeEstimator(CurrentTimecodeProvider));
+	}
+	else if (UEngineCustomTimeStep* DefaultTimestep = AskUserToCreateFallbackTimestep())
+	{
+		// There was no pre-configured custom timestep... clean it up when recording finishes.
+		RecordData.Emplace(DefaultTimestep, Private::CreateTimecodeEstimator(CurrentTimecodeProvider));
+	}
+}
+
+
+}
+
+#undef LOCTEXT_NAMESPACE
\ No newline at end of file
diff --git a/Engine/Plugins/VirtualProduction/Takes/Source/TakeRecorder/Private/Timecode/HitchlessTimecodeRecordLogic.h b/Engine/Plugins/VirtualProduction/Takes/Source/TakeRecorder/Private/Timecode/HitchlessTimecodeRecordLogic.h
new file mode 100644
--- /dev/null
+++ b/Engine/Plugins/VirtualProduction/Takes/Source/TakeRecorder/Private/Timecode/HitchlessTimecodeRecordLogic.h
@@ -0,0 +1,57 @@
+﻿// Copyright Epic Games, Inc. All Rights Reserved.
+
+#pragma once
+
+#include "GuardTimecodeProvider.h"
+#include "GuardTimestep.h"
+#include "Estimation/TimecodeRegressionProvider.h"
+#include "Misc/Optional.h"
+#include "Recorder/TakeRecorder.h"
+#include "Templates/UnrealTemplate.h"
+
+namespace UE::TakeRecorder
+{
+/**
+ * Sets up the UEngine::TimecodeProvider so timecode is not skipped when an engine frame takes longer than expected ("hitches").
+ * This is root class is responsible for setting up this for this workflow and acts as a mediator for setting up relevant systems.
+ */
+class FHitchlessTimecodeRecordLogic : public FNoncopyable
+{
+public:
+
+	FHitchlessTimecodeRecordLogic();
+	~FHitchlessTimecodeRecordLogic();
+
+private:
+
+	/** State to save for the duration of a recording. Cleans up engine state after recording. */
+	struct FRecordOverrides
+	{
+		// Member order matters: UEngineCustomTimeStep::Initialize must be run first to give it a chance to set global FApp::CurrentTime. 
+		
+		/** Set if there was no custom timestep and the user chose one in the pop-up dialogue. */
+		TOptional<FGuardTimestep> TimestepOverride;
+		/** Set while the recording and if the user did not disable it. */
+		FGuardTimecodeProvider TimecodeOverride;
+
+		FRecordOverrides(UTimecodeProvider* InOverride) : TimecodeOverride(InOverride) {}
+		FRecordOverrides(UEngineCustomTimeStep* InTimestep, UTimecodeRegressionProvider* InTimecode)
+			: TimestepOverride(InPlace, InTimestep), TimecodeOverride(InTimecode) {}
+	};
+	/** Set while the recording and if the user did not disable this feature. */
+	TOptional<FRecordOverrides> RecordData;
+	
+	/** Sets the hitchless systems once all global state is set up (specifically UTakeRecorderSubsystem) */
+	void OnPostEngineInit();
+	/** Deinits the systems; */
+	void CleanUp();
+
+	void OnInitializeRecording(UTakeRecorder* TakeRecorder) { SetupEngineEnvironmentForRecording(); }
+	void OnFinishRecording(UTakeRecorder* TakeRecorder) { RecordData.Reset(); }
+	void OnCancelRecording(UTakeRecorder* TakeRecorder) { RecordData.Reset(); }
+
+	/** Sets up timecode and optionally the custom time step. */
+	void SetupEngineEnvironmentForRecording();
+};
+}
+
diff --git a/Engine/Plugins/VirtualProduction/Takes/Source/TakeRecorder/Private/Timecode/TimecodeNotifications.cpp b/Engine/Plugins/VirtualProduction/Takes/Source/TakeRecorder/Private/Timecode/TimecodeNotifications.cpp
new file mode 100644
--- /dev/null
+++ b/Engine/Plugins/VirtualProduction/Takes/Source/TakeRecorder/Private/Timecode/TimecodeNotifications.cpp
@@ -0,0 +1,224 @@
+﻿// Copyright Epic Games, Inc. All Rights Reserved.
+
+#include "TimecodeNotifications.h"
+
+#include "CatchupFixedRateCustomTimeStep.h"
+#include "ClassViewerFilter.h"
+#include "ClassViewerModule.h"
+#include "Dialog/SCustomDialog.h"
+#include "Engine/Engine.h"
+#include "Estimation/IClockedTimeStep.h"
+#include "Framework/Application/SlateApplication.h"
+#include "Framework/Notifications/NotificationManager.h"
+#include "Modules/ModuleManager.h"
+#include "Recorder/TakeRecorderSubsystem.h"
+#include "Widgets/Input/SCheckBox.h"
+#include "Widgets/Layout/SBorder.h"
+#include "Widgets/Layout/SBox.h"
+#include "Widgets/Notifications/SNotificationList.h"
+#include "Widgets/Text/STextBlock.h"
+
+#define LOCTEXT_NAMESPACE "FHitchlessTimecodeRecordLogic"
+
+namespace UE::TakeRecorder
+{
+void ShowIncompatibleTimestepNotification(UEngineCustomTimeStep* InTimestep)
+{
+	if (FSlateApplication::IsInitialized()) // handle headless engine
+	{
+		FNotificationInfo Notification = LOCTEXT("IncompatibleTimestep.Title", "Hitchless timecode");
+		Notification.SubText = FText::Format(
+			LOCTEXT("IncompatibleTimestep.SubTextFmt", "Engine timestep {0} does not implement"),
+			InTimestep->GetClass()->GetDisplayNameText()
+			);
+		FSlateNotificationManager::Get().AddNotification(Notification)->SetCompletionState(SNotificationItem::CS_Fail);
+	}
+}
+
+namespace PickTimestepDialog
+{
+class FIsClockedTimeStepFilter : public IClassViewerFilter
+{
+public:
+	virtual bool IsClassAllowed(
+		const FClassViewerInitializationOptions& InInitOptions,
+		const UClass* InClass, TSharedRef< class FClassViewerFilterFuncs > InFilterFuncs
+		) override
+	{
+		return InClass && InClass->IsChildOf<UEngineCustomTimeStep>() && InClass->ImplementsInterface(UClockedTimeStep::StaticClass());
+	}
+
+	virtual bool IsUnloadedClassAllowed(
+		const FClassViewerInitializationOptions& InInitOptions,
+		const TSharedRef<const IUnloadedBlueprintData> InUnloadedClassData,
+		TSharedRef<FClassViewerFilterFuncs> InFilterFuncs
+		) override
+	{
+		return false;
+	}
+};
+
+static TWeakObjectPtr<UClass>& GetTimestepToUseForEditorSession()
+{
+	static TWeakObjectPtr<UClass> TimestepToUseForEditorSession;
+	return TimestepToUseForEditorSession;
+}
+
+static TSharedRef<SWidget> MakeClassPicker(UClass*& PickedClass)
+{
+	FClassViewerInitializationOptions Options;
+	Options.InitiallySelectedClass = UCatchupFixedRateCustomTimeStep::StaticClass();
+	Options.Mode = EClassViewerMode::ClassPicker;
+	Options.NameTypeToDisplay = EClassViewerNameTypeToDisplay::DisplayName;
+	Options.bShowNoneOption = false;
+	Options.bIsActorsOnly = false;
+	Options.bShowUnloadedBlueprints = false;
+	Options.ClassFilters.Add(MakeShared<FIsClockedTimeStepFilter>());
+
+	FClassViewerModule& ClassViewerModule = FModuleManager::LoadModuleChecked<FClassViewerModule>("ClassViewer");
+	return ClassViewerModule.CreateClassViewer(
+		Options,
+		FOnClassPicked::CreateLambda([&PickedClass](UClass* InNewClass)
+		{
+			PickedClass = InNewClass;
+		}));
+}
+
+static TSharedRef<SWidget> MakeDialogContent(UClass*& PickedClass, bool& bRememberChoice)
+{
+	return SNew(SBox)
+	.Padding(5, 5)
+	.HAlign(HAlign_Fill)
+	[
+		SNew(SVerticalBox)
+
+		+SVerticalBox::Slot()
+		.AutoHeight()
+		.HAlign(HAlign_Left)
+		[
+			SNew(STextBlock)
+			.Text(LOCTEXT("FixedTimeStep.InstructionMessage",
+				"No engine timestep is set up in project settings.\nPlease select a timestep class to use:"
+				))
+		]
+
+		+SVerticalBox::Slot()
+		.Padding(4, 20)
+		.AutoHeight()
+		.HAlign(HAlign_Fill)
+		[
+			SNew(SVerticalBox)
+
+			// Class picker, which only shows timesteps implementing IClockedTimeStep
+			+SVerticalBox::Slot()
+			.AutoHeight()
+			[
+				SNew(SBorder)[ MakeClassPicker(PickedClass) ]
+			]
+			
+			// Remember checkbox
+			+SVerticalBox::Slot()
+			.AutoHeight()
+			.Padding(0, 4, 20, 0)
+			.HAlign(HAlign_Left)
+			[
+				SNew(SHorizontalBox)
+
+				+SHorizontalBox::Slot()
+				.AutoWidth()
+				[
+					SNew(SCheckBox)
+					.IsChecked_Lambda([&bRememberChoice]{ return bRememberChoice ? ECheckBoxState::Checked : ECheckBoxState::Unchecked; })
+					.OnCheckStateChanged_Lambda([&bRememberChoice](ECheckBoxState NewState){ bRememberChoice = NewState == ECheckBoxState::Checked; })
+				]
+				
+				+SHorizontalBox::Slot()
+				.AutoWidth()
+				.Padding(2, 0)
+				[
+					SNew(STextBlock)
+					.Text(LOCTEXT("FixedTimeStep.Remember", "Use this class for the rest of the editor session."))
+				]
+			]
+		]
+		
+		+SVerticalBox::Slot()
+		.AutoHeight()
+		.HAlign(HAlign_Left)
+		[
+			SNew(STextBlock)
+			.Text(LOCTEXT("FixedTimeStep.DisableMessage", "To disable this workflow, execute \"TakeRecorder.EnableHitchlessTimecode 0\"."))
+		]
+	];
+}
+	
+static UClass* AskUserToSelectCustomTimestep()
+{
+	if (!FSlateApplication::IsInitialized())
+	{
+		return nullptr;
+	}
+
+	if (UClass* ClassFromLastTime = GetTimestepToUseForEditorSession().Get())
+	{
+		return ClassFromLastTime;
+	}
+
+	bool bRememberChoice = true;
+	UClass* PickedClass = nullptr;
+	
+	// This message is only temporary / for testing until we come up with better UX.
+	const TSharedRef<SCustomDialog> Message = SNew(SCustomDialog)
+		.Title(LOCTEXT("FixedTimeStep.Title", "Fixed timestep required"))
+		.Content() [ PickTimestepDialog::MakeDialogContent(PickedClass, bRememberChoice) ]
+		.Buttons(
+		{
+			SCustomDialog::FButton(LOCTEXT("FixedTimeStep.Skip.Title", "Run without hitching correction"))
+				.SetToolTipText(LOCTEXT("FixedTimeStep.Skip.Tooltip",
+					"Skip hitchless timecode set-up. "
+					"If the engine hitches, timecode evaluation will jump by the length of the frame."))
+				.SetButtonRole(SCustomDialog::EButtonRole::Cancel),
+			SCustomDialog::FButton(LOCTEXT("FixedTimeStep.DisableCVar.Title", "Disable workflow"))
+				.SetToolTipText(LOCTEXT("FixedTimeStep.Skip.Tooltip", "Executes TakeRecorder.EnableHitchlessTimecode 0 to disable hitchless setup"))
+				.SetOnClicked(FSimpleDelegate::CreateLambda([]
+				{
+					IConsoleVariable* CVar = IConsoleManager::Get().FindConsoleVariable(TEXT("TakeRecorder.EnableHitchlessTimecode"));
+					if (ensure(CVar))
+					{
+						CVar->Set(false, ECVF_SetByConsole);
+					}
+				})),
+			SCustomDialog::FButton(LOCTEXT("FixedTimeStep.UseDefault.Tooltip", "Use selected class"))
+				.SetToolTipText(LOCTEXT("FixedTimeStep.UseDefault.Title", "Uses the selected engine timestep."))
+				.SetPrimary(true)
+		});
+	
+	const int32 ButtonIndex = Message->ShowModal();
+	const bool bUseSelectedClass = ButtonIndex == 2;
+	if (bUseSelectedClass && bRememberChoice)
+	{
+		GetTimestepToUseForEditorSession() = PickedClass;
+	}
+	return bUseSelectedClass ? PickedClass : nullptr;
+}
+}
+	
+UEngineCustomTimeStep* AskUserToCreateFallbackTimestep()
+{
+	if (UClass* Class = PickTimestepDialog::AskUserToSelectCustomTimestep())
+	{
+		UEngineCustomTimeStep* CatchupTimestep = NewObject<UEngineCustomTimeStep>(GetTransientPackage(), Class);
+	
+		if (UCatchupFixedRateCustomTimeStep* Fixedtimestep = Cast<UCatchupFixedRateCustomTimeStep>(CatchupTimestep))
+		{
+			Fixedtimestep->FrameRate = GEngine->GetEngineSubsystem<UTakeRecorderSubsystem>()->GetFrameRate();
+		}
+
+		return CatchupTimestep;
+	}
+	
+	return nullptr;
+}
+}
+
+#undef LOCTEXT_NAMESPACE
\ No newline at end of file
diff --git a/Engine/Plugins/VirtualProduction/Takes/Source/TakeRecorder/Private/Timecode/TimecodeNotifications.h b/Engine/Plugins/VirtualProduction/Takes/Source/TakeRecorder/Private/Timecode/TimecodeNotifications.h
new file mode 100644
--- /dev/null
+++ b/Engine/Plugins/VirtualProduction/Takes/Source/TakeRecorder/Private/Timecode/TimecodeNotifications.h
@@ -0,0 +1,18 @@
+﻿// Copyright Epic Games, Inc. All Rights Reserved.
+
+#pragma once
+
+class UClass;
+class UEngineCustomTimeStep;
+
+namespace UE::TakeRecorder
+{
+/** Shows a notification that the current timestep is incompatible. */
+void ShowIncompatibleTimestepNotification(UEngineCustomTimeStep* InTimestep);
+	
+/**
+ * Shows a dialogue that ask the user to select a custom timestep.
+ * @return The allocated timestep but not yet registered with the engine. 
+ */
+UEngineCustomTimeStep* AskUserToCreateFallbackTimestep();
+}
