Fix static-analysis suggestions by BugHawk

diff --git a/Engine/Plugins/VirtualProduction/Takes/Source/TakeRecorder/Private/Timecode/Regression/TimecodeRegressionRecordSetup.cpp b/Engine/Plugins/VirtualProduction/Takes/Source/TakeRecorder/Private/Timecode/Regression/TimecodeRegressionRecordSetup.cpp
--- a/Engine/Plugins/VirtualProduction/Takes/Source/TakeRecorder/Private/Timecode/Regression/TimecodeRegressionRecordSetup.cpp
+++ b/Engine/Plugins/VirtualProduction/Takes/Source/TakeRecorder/Private/Timecode/Regression/TimecodeRegressionRecordSetup.cpp
@@ -37,7 +37,8 @@ static UTimecodeRegressionProvider* CreateTimecodeEstimator(UTimecodeProvider* I
 	// How many frames do we need to sample?
 	// E.g. if FrameRate is 60 FPS of 60 frames, then Numerator = 60 meaning 60 / 1 = 60. We want
 	const FFrameRate FrameRate = GEngine->GetEngineSubsystem<UTakeRecorderSubsystem>()->GetFrameRate();
-	TimecodeRegression->NumSampledFrames = CVarTimecodeLinearRegressionSeconds.GetValueOnAnyThread() * FrameRate.Numerator / FrameRate.Denominator;
+	TimecodeRegression->NumSampledFrames = FrameRate.Denominator == 0
+		? 1 : CVarTimecodeLinearRegressionSeconds.GetValueOnAnyThread() * FrameRate.Numerator / FrameRate.Denominator;
 	
 	return TimecodeRegression;
 }
diff --git a/Engine/Source/Runtime/TimeManagement/Private/Misc/LinearRegression.cpp b/Engine/Source/Runtime/TimeManagement/Private/Misc/LinearRegression.cpp
--- a/Engine/Source/Runtime/TimeManagement/Private/Misc/LinearRegression.cpp
+++ b/Engine/Source/Runtime/TimeManagement/Private/Misc/LinearRegression.cpp
@@ -10,7 +10,7 @@ bool ComputeLinearRegressionSlopeAndOffset(
 	)
 {
 	// This happens when all data points are equal. The independent variables are all on the same point. There is no slope.
-	const double Variance = InArgs.SumOfSquaredXes - (FMath::Square(InArgs.SumX) / InArgs.Num);
+	const double Variance = InArgs.Num == 0 ? 0.0 : InArgs.SumOfSquaredXes - (FMath::Square(InArgs.SumX) / InArgs.Num);
 	if (FMath::IsNearlyZero(Variance))
 	{
 		OutSlope = 0.0;
@@ -40,7 +40,7 @@ FLinearRegressionArgs ComputeLinearRegressionInputArgs(const TConstArrayView<FVe
 	{
 		Result.SumX += Sample.X;
 		Result.SumY += Sample.Y;
-		Result.SumOfSquaredXes += Sample.X * Sample.Y;
+		Result.SumXxY += Sample.X * Sample.Y;
 		Result.SumOfSquaredXes += FMath::Square(Sample.X);
 	}
 	return Result;
