Fix localization build warnings

diff --git a/Engine/Plugins/VirtualProduction/Takes/Source/TakeRecorder/Private/Timecode/Regression/TimecodeNotifications.cpp b/Engine/Plugins/VirtualProduction/Takes/Source/TakeRecorder/Private/Timecode/Regression/TimecodeNotifications.cpp
--- a/Engine/Plugins/VirtualProduction/Takes/Source/TakeRecorder/Private/Timecode/Regression/TimecodeNotifications.cpp
+++ b/Engine/Plugins/VirtualProduction/Takes/Source/TakeRecorder/Private/Timecode/Regression/TimecodeNotifications.cpp
@@ -179,7 +179,7 @@ static UClass* AskUserToSelectCustomTimestep()
 					"If the engine hitches, timecode evaluation will jump by the length of the frame."))
 				.SetButtonRole(SCustomDialog::EButtonRole::Cancel),
 			SCustomDialog::FButton(LOCTEXT("FixedTimeStep.DisableCVar.Title", "Disable workflow"))
-				.SetToolTipText(LOCTEXT("FixedTimeStep.Skip.Tooltip", "Executes TakeRecorder.EnableHitchlessTimecode 0 to disable hitchless setup"))
+				.SetToolTipText(LOCTEXT("FixedTimeStep.DisableCVar.Tooltip", "Executes TakeRecorder.EnableHitchlessTimecode 0 to disable hitchless setup"))
 				.SetOnClicked(FSimpleDelegate::CreateLambda([]
 				{
 					IConsoleVariable* CVar = IConsoleManager::Get().FindConsoleVariable(TEXT("TakeRecorder.EnableHitchlessTimecode"));
diff --git a/Engine/Plugins/VirtualProduction/Takes/Source/TakeRecorder/Private/Timecode/Regression/TimecodeRegressionRecordSetup.cpp b/Engine/Plugins/VirtualProduction/Takes/Source/TakeRecorder/Private/Timecode/Regression/TimecodeRegressionRecordSetup.cpp
--- a/Engine/Plugins/VirtualProduction/Takes/Source/TakeRecorder/Private/Timecode/Regression/TimecodeRegressionRecordSetup.cpp
+++ b/Engine/Plugins/VirtualProduction/Takes/Source/TakeRecorder/Private/Timecode/Regression/TimecodeRegressionRecordSetup.cpp
@@ -113,8 +113,4 @@ void FTimecodeRegressionRecordSetup::SetupEngineEnvironmentForRecording()
 		RecordData.Emplace(DefaultTimestep, Private::CreateTimecodeEstimator(CurrentTimecodeProvider));
 	}
 }
-
-
-}
-
-#undef LOCTEXT_NAMESPACE
\ No newline at end of file
+}
\ No newline at end of file
