In non-editor worlds, we should defer to the record type versus what the global spawnable specifices.  This is to address cases where we change the recording type on pre-recording

#rnx

diff --git a/Engine/Plugins/VirtualProduction/Takes/Source/TakeRecorderSources/Private/TakeRecorderActorSource.cpp b/Engine/Plugins/VirtualProduction/Takes/Source/TakeRecorderSources/Private/TakeRecorderActorSource.cpp
--- a/Engine/Plugins/VirtualProduction/Takes/Source/TakeRecorderSources/Private/TakeRecorderActorSource.cpp
+++ b/Engine/Plugins/VirtualProduction/Takes/Source/TakeRecorderSources/Private/TakeRecorderActorSource.cpp
@@ -1778,7 +1778,9 @@ bool UTakeRecorderActorSource::GetRecordToPossessable() const
 		return true;
 	}
 	
-	if (!AllowsSpawnableObjects())
+	// In non-editor worlds, we should defer to the record type versus what the global spawnable specifices.  This is to address cases where we change
+	// the recording type on pre-recording. See UE-306474.
+	if (Target->GetWorld()->WorldType == EWorldType::Editor && !AllowsSpawnableObjects())
 	{
 		return true;
 	}
