TakeRecorder: Fix custom time dilation failing when bStartAtCurrentTimecode is true and sequencer is disabled.

diff --git a/Engine/Plugins/VirtualProduction/Takes/Source/TakeRecorder/Private/Recorder/TakeRecorder.cpp b/Engine/Plugins/VirtualProduction/Takes/Source/TakeRecorder/Private/Recorder/TakeRecorder.cpp
--- a/Engine/Plugins/VirtualProduction/Takes/Source/TakeRecorder/Private/Recorder/TakeRecorder.cpp
+++ b/Engine/Plugins/VirtualProduction/Takes/Source/TakeRecorder/Private/Recorder/TakeRecorder.cpp
@@ -786,7 +786,21 @@ void UTakeRecorder::Tick(float DeltaTime)
 
 FQualifiedFrameTime UTakeRecorder::GetRecordTime() const
 {
-	return TakesUtils::GetRecordTime(WeakSequencer.Pin(), SequenceAsset, TimecodeAtStart, Parameters.Project.bStartAtCurrentTimecode);
+	float TimeDilation = 1.f;
+	if (!WeakSequencer.IsValid())
+	{
+		const UWorld* RecordingWorld = GetWorld();
+		check(RecordingWorld);
+		if (const AWorldSettings* WorldSettings = RecordingWorld->GetWorldSettings())
+		{
+			// If we're not using Sequencer we have to apply time dilation ourselves as this normally happens in Sequencer.cpp
+			// when the TimeController is ticked.
+			TimeDilation = WorldSettings->CinematicTimeDilation;
+		}
+	}
+	
+	return TakesUtils::GetRecordTime(WeakSequencer.Pin(), SequenceAsset, TimecodeAtStart,
+		Parameters.Project.bStartAtCurrentTimecode, TimeDilation);
 }
 
 void UTakeRecorder::InternalTick(float DeltaTime)
@@ -797,24 +811,8 @@ void UTakeRecorder::InternalTick(float DeltaTime)
 	FScopedSignedObjectModifyDefer FlushOnTick(true);
 
 	TSharedPtr<ISequencer> Sequencer = WeakSequencer.Pin();
-	FQualifiedFrameTime RecordTime = GetRecordTime();
 
-	if (!Sequencer.IsValid())
-	{
-		const UWorld* RecordingWorld = GetWorld();
-		check(RecordingWorld);
-		if (const AWorldSettings* WorldSettings = RecordingWorld->GetWorldSettings())
-		{
-			// If we're not using Sequencer we have to apply time dilation ourselves as this normally happens in Sequencer.cpp
-			// when the TimeController is ticked.
-			FQualifiedFrameTime ScaledRecordTime(
-				RecordTime.Time * WorldSettings->CinematicTimeDilation,
-				RecordTime.Rate
-			);
-
-			RecordTime = MoveTemp(ScaledRecordTime);
-		}
-	}
+	FQualifiedFrameTime RecordTime = GetRecordTime();
 	
 	UTakeRecorderSources* Sources = SequenceAsset->FindOrAddMetaData<UTakeRecorderSources>();
 	if (!Parameters.bDisableRecordingAndSave)
diff --git a/Engine/Plugins/VirtualProduction/Takes/Source/TakesCore/Private/TakesUtils.cpp b/Engine/Plugins/VirtualProduction/Takes/Source/TakesCore/Private/TakesUtils.cpp
--- a/Engine/Plugins/VirtualProduction/Takes/Source/TakesCore/Private/TakesUtils.cpp
+++ b/Engine/Plugins/VirtualProduction/Takes/Source/TakesCore/Private/TakesUtils.cpp
@@ -193,34 +193,38 @@ TSharedPtr<ISequencer> OpenSequencer(ULevelSequence* LevelSequence, FText* OutEr
 	return Sequencer;
 }
 
-FQualifiedFrameTime GetRecordTime(TSharedPtr<ISequencer> Sequencer, ULevelSequence* SequenceAsset, const FTimecode& TimecodeAtStart, bool bStartAtCurrentTimecode)
+FQualifiedFrameTime GetRecordTime(TSharedPtr<ISequencer> Sequencer, ULevelSequence* SequenceAsset,
+	const FTimecode& TimecodeAtStart, bool bStartAtCurrentTimecode, float TimeDilation)
 {
 	FQualifiedFrameTime RecordTime;
 
 	if (Sequencer.IsValid())
 	{
+		// Sequencer handles time dilation.
 		RecordTime = Sequencer->GetGlobalTime();
 	}
 	else if (SequenceAsset)
 	{
 		if (UMovieScene* MovieScene = SequenceAsset->GetMovieScene())
 		{
-			FFrameRate FrameRate = MovieScene->GetDisplayRate();
-			FFrameRate TickResolution = MovieScene->GetTickResolution();
+			const FFrameRate FrameRate = MovieScene->GetDisplayRate();
+			const FFrameRate TickResolution = MovieScene->GetTickResolution();
 
-			FTimecode CurrentTimecode = FApp::GetTimecode();
-		
-			FFrameNumber CurrentFrame = FFrameRate::TransformTime(FFrameTime(CurrentTimecode.ToFrameNumber(FrameRate)), FrameRate, TickResolution).FloorToFrame();
-			FFrameNumber FrameAtStart = FFrameRate::TransformTime(FFrameTime(TimecodeAtStart.ToFrameNumber(FrameRate)), FrameRate, TickResolution).FloorToFrame();
+			const FTimecode CurrentTimecode = FApp::GetTimecode();
 
-			if (bStartAtCurrentTimecode)
-			{
-				RecordTime = FQualifiedFrameTime(CurrentFrame, TickResolution);
-			}
-			else
-			{
-				RecordTime = FQualifiedFrameTime(CurrentFrame - FrameAtStart, TickResolution);
-			}
+			const FFrameTime CurrentFrameTime = FFrameRate::TransformTime(FFrameTime(CurrentTimecode.ToFrameNumber(FrameRate)), FrameRate, TickResolution);
+			const FFrameTime StartFrameTime = FFrameRate::TransformTime(FFrameTime(TimecodeAtStart.ToFrameNumber(FrameRate)), FrameRate, TickResolution);
+
+			// If we are managing time dilation, then we need to apply it to the delta only.
+			const FFrameTime DeltaFrameTime = CurrentFrameTime - StartFrameTime;
+			const FFrameTime DilatedFrameTime = DeltaFrameTime * TimeDilation;
+
+			// Offset with the start time after dilation has been applied, if applicable.
+			const FFrameTime FinalFrameTime = bStartAtCurrentTimecode
+				? StartFrameTime + DilatedFrameTime
+				: DilatedFrameTime;
+			
+			RecordTime = FQualifiedFrameTime(FinalFrameTime, TickResolution);
 		}
 	}
 
diff --git a/Engine/Plugins/VirtualProduction/Takes/Source/TakesCore/Public/TakesUtils.h b/Engine/Plugins/VirtualProduction/Takes/Source/TakesCore/Public/TakesUtils.h
--- a/Engine/Plugins/VirtualProduction/Takes/Source/TakesCore/Public/TakesUtils.h
+++ b/Engine/Plugins/VirtualProduction/Takes/Source/TakesCore/Public/TakesUtils.h
@@ -171,5 +171,6 @@ namespace TakesUtils
 	TAKESCORE_API TSharedPtr<ISequencer> OpenSequencer(ULevelSequence* LevelSequence, FText* OutError);
 	
 	/* The time at which to record. Taken from the Sequencer global time, otherwise based on timecode */
-	TAKESCORE_API FQualifiedFrameTime GetRecordTime(TSharedPtr<ISequencer> Sequencer, ULevelSequence* SequenceAsset, const FTimecode& TimecodeAtStart, bool bStartAtCurrentTimecode);
+	TAKESCORE_API FQualifiedFrameTime GetRecordTime(TSharedPtr<ISequencer> Sequencer, ULevelSequence* SequenceAsset,
+		const FTimecode& TimecodeAtStart, bool bStartAtCurrentTimecode, float TimeDilation = 1.f);
 }
\ No newline at end of file
