Prevent a crash when working with an invalid source which can happen in Multi-user scenarios.

#rnx

diff --git a/Engine/Plugins/VirtualProduction/Takes/Source/TakesCore/Private/TakeRecorderSources.cpp b/Engine/Plugins/VirtualProduction/Takes/Source/TakesCore/Private/TakeRecorderSources.cpp
--- a/Engine/Plugins/VirtualProduction/Takes/Source/TakesCore/Private/TakeRecorderSources.cpp
+++ b/Engine/Plugins/VirtualProduction/Takes/Source/TakesCore/Private/TakeRecorderSources.cpp
@@ -103,6 +103,18 @@ void UTakeRecorderSources::PostEditChangeProperty(FPropertyChangedEvent& Propert
 	{
 		++SourcesSerialNumber;
 	}
+
+	// There have been cases where users are transacted sources to remote nodes with Multi-user and those sources cannot be assigned due to missing objects.
+	// For example, creating a audio source fails on a remote nDisplay node because of missing audio objects.  We want to flag these in the log.
+	int32 SlotIndex = 0;
+	for (UTakeRecorderSource* Source : Sources)
+	{
+		if (!Source)
+		{
+			UE_LOG(LogTakesCore, Warning, TEXT("Transacted source at slot %d is invalid will be ignored during recording."), SlotIndex);
+		}
+		SlotIndex++;
+	}
 }
 
 void UTakeRecorderSources::PreRecordingRecursive(TArray<UTakeRecorderSource*> InSources, ULevelSequence* InRootSequence, TArray<UTakeRecorderSource*>& NewSourcesOut, FManifestSerializer* InManifestSerializer)
@@ -118,7 +130,7 @@ void UTakeRecorderSources::PreRecordingRecursive(TArray<UTakeRecorderSource*> In
 
 	for (UTakeRecorderSource* Source : InSources)
 	{
-		if (Source->bEnabled)
+		if (Source && Source->bEnabled)
 		{
 			ULevelSequence* TargetSequence = InRootSequence;
 			FMovieSceneSequenceID TargetSequenceID = MovieSceneSequenceID::Root;
@@ -290,7 +302,7 @@ void UTakeRecorderSources::StartRecordingTheseSources(const TArray<UTakeRecorder
 
 	for (auto Source : InSources)
 	{
-		if (Source->bEnabled)
+		if (Source && Source->bEnabled)
 		{
 			if (Settings.bRecordSourcesIntoSubSequences && Source->SupportsSubscenes()) //Set Timcode on MovieScene if we created a sub scene for it
 			{
@@ -376,7 +388,7 @@ FFrameTime UTakeRecorderSources::TickRecording(class ULevelSequence* InSequence,
 	{
 		for (int32 SourceIndex = 0; SourceIndex < Sources.Num(); ++SourceIndex)
 		{
-			if (Sources[SourceIndex]->bEnabled)
+			if (Sources[SourceIndex] && Sources[SourceIndex]->bEnabled)
 			{
 				Sources[SourceIndex]->TickRecording(CurrentFrameTime);
 			}
@@ -410,7 +422,7 @@ void UTakeRecorderSources::StopRecording(class ULevelSequence* InSequence, const
 
 	for (auto Source : Sources)
 	{
-		if (Source->bEnabled)
+		if (Source && Source->bEnabled)
 		{
 			Source->StopRecording(SourceSubSequenceMap[Source]);
 		}
@@ -419,7 +431,7 @@ void UTakeRecorderSources::StopRecording(class ULevelSequence* InSequence, const
 	TArray<UTakeRecorderSource*> SourcesToRemove;
 	for (auto Source : Sources)
 	{
-		if (Source->bEnabled)
+		if (Source && Source->bEnabled)
 		{
 			for (auto SourceToRemove : Source->PostRecording(SourceSubSequenceMap[Source], InSequence, bCancelled))
 			{
@@ -439,7 +451,10 @@ void UTakeRecorderSources::StopRecording(class ULevelSequence* InSequence, const
 
 	for (UTakeRecorderSource* Source : Sources)
 	{
-		Source->FinalizeRecording();
+		if (Source)
+		{
+			Source->FinalizeRecording();
+		}
 	}
 
 	// Re-enable transactional after recording
