Create a stat category for the take recorder and add stat item for displaying the take recorder sources.

diff --git a/Engine/Plugins/VirtualProduction/Takes/Source/TakesCore/Private/TakeRecorderSources.cpp b/Engine/Plugins/VirtualProduction/Takes/Source/TakesCore/Private/TakeRecorderSources.cpp
--- a/Engine/Plugins/VirtualProduction/Takes/Source/TakesCore/Private/TakeRecorderSources.cpp
+++ b/Engine/Plugins/VirtualProduction/Takes/Source/TakesCore/Private/TakeRecorderSources.cpp
@@ -24,6 +24,13 @@
 #include "NamingTokens/TakeRecorderNamingTokensContext.h"
 #include "UObject/Package.h"
 
+#include "Engine/Engine.h"
+#include "Engine/Font.h"
+#include "CanvasTypes.h"
+
+
+#define LOCTEXT_NAMESPACE "TakesUtils"
+
 #include UE_INLINE_GENERATED_CPP_BY_NAME(TakeRecorderSources)
 
 DEFINE_LOG_CATEGORY(SubSequenceSerialization);
@@ -39,6 +46,15 @@ UTakeRecorderSources::UTakeRecorderSources(const FObjectInitializer& ObjInit)
 	{
 		SetFlags(RF_Transactional);
 	}
+
+	if (ensure(GEngine))
+	{
+		constexpr bool bIsRHS = true;
+		GEngine->AddEngineStat(TEXT("STAT_TakeRecorderSources"), TEXT("STATCAT_TakeRecorder"),
+							   LOCTEXT("DisplayTakeRecorderSources", "Displays current take recorder sources in the stats panel."),
+							   UEngine::FEngineStatRender::CreateUObject(this, &UTakeRecorderSources::RenderStatTakeSources),
+							   nullptr, bIsRHS);
+	}
 }
 
 UTakeRecorderSource* UTakeRecorderSources::AddSource(TSubclassOf<UTakeRecorderSource> InSourceType)
@@ -759,3 +775,43 @@ void UTakeRecorderSources::RemoveRedundantTracks()
 	}
 }
 
+
+int32 UTakeRecorderSources::RenderStatTakeSources(UWorld* World, FViewport* Viewport, FCanvas* Canvas, int32 X, int32 Y, const FVector* ViewLocation, const FRotator* ViewRotation)
+{
+	check(GEngine);
+	UFont* Font = FPlatformProperties::SupportsWindowedMode() ? GEngine->GetSmallFont() : GEngine->GetMediumFont();
+	const int32 RowHeight = FMath::TruncToInt(Font->GetMaxCharHeight());
+
+	float CharWidth, CharHeight;
+	Font->GetCharSize(TEXT(' '), CharWidth, CharHeight);
+	auto RenderStatForSource = [RowHeight, World, Viewport, Font, Canvas, CharWidth, X, &Y](UTakeRecorderSource* Source) mutable
+	{
+		FText DisplayText = Source->GetDisplayText();
+		FString DrawString = FString::Printf(TEXT("TakeRecorderSource: %s"), *DisplayText.ToString());
+		X = X - Font->GetStringSize(*DrawString);
+		Canvas->DrawShadowedString(X, Y, *DrawString, Font, FColor::Green);
+		Y += RowHeight;
+	};
+
+	int32 NumSources = 0;
+	for (UTakeRecorderSource* Source : Sources)
+	{
+		if (Source && Source->bEnabled)
+		{
+			RenderStatForSource(Source);
+			NumSources ++;
+		}
+	}
+
+	if (NumSources == 0)
+	{
+		FString DisplayText = TEXT("No Take Recorder sources assigned");
+		X = X - Font->GetStringSize(*DisplayText);
+
+		Canvas->DrawShadowedString(X, Y, *DisplayText, Font, FColor::Red);
+		Y += RowHeight;
+	}
+	return Y;
+}
+
+#undef LOCTEXT_NAMESPACE
diff --git a/Engine/Plugins/VirtualProduction/Takes/Source/TakesCore/Public/TakeRecorderSources.h b/Engine/Plugins/VirtualProduction/Takes/Source/TakesCore/Public/TakeRecorderSources.h
--- a/Engine/Plugins/VirtualProduction/Takes/Source/TakesCore/Public/TakeRecorderSources.h
+++ b/Engine/Plugins/VirtualProduction/Takes/Source/TakesCore/Public/TakeRecorderSources.h
@@ -30,6 +30,10 @@ struct FTakeRecorderSourcesSettings
 	FString SubSequenceDirectory;
 };
 
+class UWorld;
+class FViewport;
+class FCanvas;
+
 /**
  * A list of sources to record for any given take. Stored as meta-data on ULevelSequence through ULevelSequence::FindMetaData<UTakeRecorderSources>
  */
@@ -187,6 +191,12 @@ public:
 	FQualifiedFrameTime GetCachedFrameTime() const { return CachedFrameTime; }
 
 private:
+	/**
+	 * Renders stat text that displays the current list of active take recorder sources.  The take recorder can now operate without a panel and this
+	 * gives users a way to view the current sources list outside of the panel.
+	 */
+	int32 RenderStatTakeSources(UWorld* World, FViewport* Viewport, FCanvas* Canvas, int32 X, int32 Y, const FVector* ViewLocation, const FRotator* ViewRotation);
+		
 	/** Called at the end of each frame in both the Editor and in Game to update all Sources. */
 	virtual void Tick(float DeltaTime) {}
 
