Remove ensure it was causing cook issues.

diff --git a/Engine/Plugins/VirtualProduction/Takes/Source/TakesCore/Private/TakeRecorderSources.cpp b/Engine/Plugins/VirtualProduction/Takes/Source/TakesCore/Private/TakeRecorderSources.cpp
--- a/Engine/Plugins/VirtualProduction/Takes/Source/TakesCore/Private/TakeRecorderSources.cpp
+++ b/Engine/Plugins/VirtualProduction/Takes/Source/TakesCore/Private/TakeRecorderSources.cpp
@@ -47,7 +47,7 @@ UTakeRecorderSources::UTakeRecorderSources(const FObjectInitializer& ObjInit)
 		SetFlags(RF_Transactional);
 	}
 
-	if (ensure(GEngine))
+	if (GEngine)
 	{
 		constexpr bool bIsRHS = true;
 		GEngine->AddEngineStat(TEXT("STAT_TakeRecorderSources"), TEXT("STATCAT_TakeRecorder"),
