Move ProcessRecordedTimes to TakesCore. This allows the function to be used in TakeRecorder in a follow-up change.

diff --git a/Engine/Plugins/Animation/LiveLink/Source/LiveLinkSequencer/Private/SequencerRecorderSections/MovieSceneLiveLinkTrackRecorder.cpp b/Engine/Plugins/Animation/LiveLink/Source/LiveLinkSequencer/Private/SequencerRecorderSections/MovieSceneLiveLinkTrackRecorder.cpp
--- a/Engine/Plugins/Animation/LiveLink/Source/LiveLinkSequencer/Private/SequencerRecorderSections/MovieSceneLiveLinkTrackRecorder.cpp
+++ b/Engine/Plugins/Animation/LiveLink/Source/LiveLinkSequencer/Private/SequencerRecorderSections/MovieSceneLiveLinkTrackRecorder.cpp
@@ -16,6 +16,7 @@
 #include "MovieSceneTakeTrack.h"
 #include "TakeRecorderSettings.h"
 #include "TakeRecorderSourceHelpers.h"
+#include "TakeRecorderTimeProcessing.h"
 
 #include "Trace/Trace.inl"
 
@@ -238,7 +239,7 @@ void UMovieSceneLiveLinkTrackRecorder::ProcessRecordedTimes(ULevelSequence* InLe
 	{
 		TakeTrack = Cast<UMovieSceneTakeTrack>(MovieScene->AddTrack(UMovieSceneTakeTrack::StaticClass()));
 	}
-	TakeRecorderSourceHelpers::ProcessRecordedTimes(InLevelSequence, TakeTrack, FrameRange, RecordedTimes);
+	UE::TakesCore::ProcessRecordedTimes(InLevelSequence, TakeTrack, FrameRange, RecordedTimes);
 }
 
 void UMovieSceneLiveLinkTrackRecorder::RecordSampleImpl(const FQualifiedFrameTime& CurrentTime)
diff --git a/Engine/Plugins/VirtualProduction/Takes/Source/TakeRecorderSources/Private/TakeRecorderActorSource.cpp b/Engine/Plugins/VirtualProduction/Takes/Source/TakeRecorderSources/Private/TakeRecorderActorSource.cpp
--- a/Engine/Plugins/VirtualProduction/Takes/Source/TakeRecorderSources/Private/TakeRecorderActorSource.cpp
+++ b/Engine/Plugins/VirtualProduction/Takes/Source/TakeRecorderSources/Private/TakeRecorderActorSource.cpp
@@ -51,6 +51,7 @@
 #include "NamingTokens/TakeRecorderNamingTokensContext.h"
 #include "SequencerUtilities.h"
 #include "TakeRecorderSourcesModule.h"
+#include "TakeRecorderTimeProcessing.h"
 
 #include UE_INLINE_GENERATED_CPP_BY_NAME(TakeRecorderActorSource)
 
@@ -740,7 +741,7 @@ void UTakeRecorderActorSource::ProcessRecordedTimes(ULevelSequence* InSequence)
 	UMovieSceneTakeTrack* TakeTrack = UE::TakeRecorderActorSource::Private::FindOrAddTakeTrack(MovieScene, CachedObjectBindingGuid);
 
 	// Add the recorded times to the take track.
-	TakeRecorderSourceHelpers::ProcessRecordedTimes(InSequence, TakeTrack, FrameRange, UTakeRecorderSources::RecordedTimes);
+	UE::TakesCore::ProcessRecordedTimes(InSequence, TakeTrack, FrameRange, UTakeRecorderSources::RecordedTimes);
 }
 
 UMovieSceneAnimationTrackRecorderEditorSettings* UTakeRecorderActorSource::GetBestTrackRecorderEditorSettings() const
diff --git a/Engine/Plugins/VirtualProduction/Takes/Source/TakeRecorderSources/Private/TakeRecorderSourceHelpers.cpp b/Engine/Plugins/VirtualProduction/Takes/Source/TakeRecorderSources/Private/TakeRecorderSourceHelpers.cpp
--- a/Engine/Plugins/VirtualProduction/Takes/Source/TakeRecorderSources/Private/TakeRecorderSourceHelpers.cpp
+++ b/Engine/Plugins/VirtualProduction/Takes/Source/TakeRecorderSources/Private/TakeRecorderSourceHelpers.cpp
@@ -7,12 +7,10 @@
 #include "GameFramework/Actor.h"
 #include "ScopedTransaction.h"
 #include "TakeRecorderLevelSequenceSource.h"
+#include "TakeRecorderTimeProcessing.h"
 
 #include "LevelSequence.h"
-#include "MovieScene.h"
-#include "MovieSceneTakeTrack.h"
 #include "MovieSceneTakeSection.h"
-#include "TakeMetaData.h"
 
 #include "Editor.h"
 
@@ -143,121 +141,12 @@ AActor* GetSourceActor(UTakeRecorderSource* Source)
 	return nullptr;
 }
 
-void ProcessRecordedTimes(ULevelSequence* InSequence, UMovieSceneTakeTrack* TakeTrack, const TOptional<TRange<FFrameNumber>>& FrameRange, const FArrayOfRecordedTimePairs& RecordedTimes)
+void ProcessRecordedTimes(
+	ULevelSequence* InSequence, UMovieSceneTakeTrack* TakeTrack,
+	const TOptional<TRange<FFrameNumber>>& FrameRange, const TArray<TPair<FQualifiedFrameTime, FQualifiedFrameTime>>& RecordedTimes
+	)
 {
-	check(TakeTrack);
-	UMovieScene* MovieScene = InSequence->GetMovieScene();
-
-	// In case we need it later, get the earliest timecode source *before* we
-	// add the take section, since its timecode source will be default
-	// constructed as all zeros and might accidentally compare as earliest.
-	const FMovieSceneTimecodeSource EarliestTimecodeSource = MovieScene->GetEarliestTimecodeSource();
-
-	TakeTrack->RemoveAllAnimationData();
-
-	UMovieSceneTakeSection* TakeSection = Cast<UMovieSceneTakeSection>(TakeTrack->CreateNewSection());
-	TakeTrack->AddSection(*TakeSection);
-
-	if (FrameRange.IsSet())
-	{
-		TArray<int32> Hours, Minutes, Seconds, Frames;
-		TArray<FMovieSceneFloatValue> SubFrames;
-		TArray<FFrameNumber> Times;
-
-		Hours.Reserve(RecordedTimes.Num());
-		Minutes.Reserve(RecordedTimes.Num());
-		Seconds.Reserve(RecordedTimes.Num());
-		Frames.Reserve(RecordedTimes.Num());
-		SubFrames.Reserve(RecordedTimes.Num());
-		Times.Reserve(RecordedTimes.Num());
-
-		FFrameRate TickResolution = MovieScene->GetTickResolution();
-		FFrameRate DisplayRate = MovieScene->GetDisplayRate();
-
-		FFrameRate TCRate = TickResolution;
-		for (const TPair<FQualifiedFrameTime, FQualifiedFrameTime>& RecordedTimePair : RecordedTimes)
-		{
-			FFrameNumber FrameNumber = RecordedTimePair.Key.Time.FrameNumber;
-			if (!FrameRange.GetValue().Contains(FrameNumber))
-			{
-				continue;
-			}
-
-			FTimecode Timecode = RecordedTimePair.Value.ToTimecode();
-			TCRate = RecordedTimePair.Value.Rate;
-			Hours.Add(Timecode.Hours);
-			Minutes.Add(Timecode.Minutes);
-			Seconds.Add(Timecode.Seconds);
-			Frames.Add(Timecode.Frames);
-
-			FMovieSceneFloatValue SubFrame;
-			if (RecordedTimePair.Value.Time.GetSubFrame() > 0)
-			{
-				// If the Timecode provided gave us a subframe value then we should use that value.  Otherwise, we should compute
-				// the most appropriate value based on the timecode rate.
-				SubFrame.Value = RecordedTimePair.Value.Time.GetSubFrame();
-			}
-			else
-			{
-				FFrameTime FrameTime = FFrameRate::TransformTime(RecordedTimePair.Key.Time, TickResolution, DisplayRate);
-				FQualifiedFrameTime FrameTimeAsTimeCodeRate(FrameTime, TCRate);
-
-				SubFrame.Value = FrameTimeAsTimeCodeRate.Time.GetSubFrame();
-			}
-
-			SubFrame.InterpMode = ERichCurveInterpMode::RCIM_Linear;
-			SubFrames.Add(SubFrame);
-
-			Times.Add(FrameNumber);
-		}
-
-		Hours.Shrink();
-		Minutes.Shrink();
-		Seconds.Shrink();
-		Frames.Shrink();
-		SubFrames.Shrink();
-		Times.Shrink();
-
-		TakeSection->HoursCurve.Set(Times, Hours);
-		TakeSection->MinutesCurve.Set(Times, Minutes);
-		TakeSection->SecondsCurve.Set(Times, Seconds);
-		TakeSection->FramesCurve.Set(Times, Frames);
-		TakeSection->SubFramesCurve.Set(Times, SubFrames);
-		TakeSection->RateCurve.SetDefault(TCRate.AsDecimal());
-	}
-
-	// Since the take section was created post recording here in this
-	// function, it wasn't available at the start of recording to have
-	// its timecode source set with the other sections, so we set it here.
-	if (TakeSection->HoursCurve.GetNumKeys() > 0)
-	{
-		// We populated the take section's timecode curves with data, so
-		// use the first values as the timecode source.
-		const int32 Hours = TakeSection->HoursCurve.GetValues()[0];
-		const int32 Minutes = TakeSection->MinutesCurve.GetValues()[0];
-		const int32 Seconds = TakeSection->SecondsCurve.GetValues()[0];
-		const int32 Frames = TakeSection->FramesCurve.GetValues()[0];
-		const bool bIsDropFrame = false;
-		const FTimecode Timecode(Hours, Minutes, Seconds, Frames, bIsDropFrame);
-		TakeSection->TimecodeSource = FMovieSceneTimecodeSource(Timecode);
-	}
-	else
-	{
-		// Otherwise, adopt the earliest timecode source from one of the movie
-		// scene's other sections as the timecode source for the take section.
-		// This case is unlikely.
-		TakeSection->TimecodeSource = EarliestTimecodeSource;
-	}
-
-	if (UTakeMetaData* TakeMetaData = InSequence->FindMetaData<UTakeMetaData>())
-	{
-		TakeSection->Slate.SetDefault(FString::Printf(TEXT("%s_%d"), *TakeMetaData->GetSlate(), TakeMetaData->GetTakeNumber()));
-	}
-
-	if (TakeSection->GetAutoSizeRange().IsSet())
-	{
-		TakeSection->SetRange(TakeSection->GetAutoSizeRange().GetValue());
-	}
+	return UE::TakesCore::ProcessRecordedTimes(InSequence, TakeTrack, FrameRange, RecordedTimes);
 }
 
 }
diff --git a/Engine/Plugins/VirtualProduction/Takes/Source/TakeRecorderSources/Public/TakeRecorderSourceHelpers.h b/Engine/Plugins/VirtualProduction/Takes/Source/TakeRecorderSources/Public/TakeRecorderSourceHelpers.h
--- a/Engine/Plugins/VirtualProduction/Takes/Source/TakeRecorderSources/Public/TakeRecorderSourceHelpers.h
+++ b/Engine/Plugins/VirtualProduction/Takes/Source/TakeRecorderSources/Public/TakeRecorderSourceHelpers.h
@@ -51,5 +51,6 @@ namespace TakeRecorderSourceHelpers
 	/**
 	 * Creates a takes track to store timecode data on a take recorder source.
 	 */
+	UE_DEPRECATED(5.7, "Use UE::TakeCores::ProcessRecordedTimes instead.")
 	TAKERECORDERSOURCES_API void ProcessRecordedTimes(ULevelSequence* InSequence, UMovieSceneTakeTrack* TakeTrack, const TOptional<TRange<FFrameNumber>>& FrameRange, const FArrayOfRecordedTimePairs& RecordedTimes);
 };
diff --git a/Engine/Plugins/VirtualProduction/Takes/Source/TakesCore/Private/TakeRecorderTimeProcessing.cpp b/Engine/Plugins/VirtualProduction/Takes/Source/TakesCore/Private/TakeRecorderTimeProcessing.cpp
new file mode 100644
--- /dev/null
+++ b/Engine/Plugins/VirtualProduction/Takes/Source/TakesCore/Private/TakeRecorderTimeProcessing.cpp
@@ -0,0 +1,132 @@
+﻿// Copyright Epic Games, Inc. All Rights Reserved.
+
+#include "TakeRecorderTimeProcessing.h"
+
+#include "LevelSequence.h"
+#include "MovieScene.h"
+#include "MovieSceneTakeSection.h"
+#include "MovieSceneTakeTrack.h"
+#include "TakeMetaData.h"
+#include "Channels/MovieSceneFloatChannel.h"
+#include "Curves/RealCurve.h"
+
+namespace UE::TakesCore
+{
+void ProcessRecordedTimes(ULevelSequence* InSequence, UMovieSceneTakeTrack* TakeTrack, const TOptional<TRange<FFrameNumber>>& FrameRange,
+	const FArrayOfRecordedTimePairs& RecordedTimes)
+{
+	check(TakeTrack);
+	UMovieScene* MovieScene = InSequence->GetMovieScene();
+
+	// In case we need it later, get the earliest timecode source *before* we
+	// add the take section, since its timecode source will be default
+	// constructed as all zeros and might accidentally compare as earliest.
+	const FMovieSceneTimecodeSource EarliestTimecodeSource = MovieScene->GetEarliestTimecodeSource();
+
+	TakeTrack->RemoveAllAnimationData();
+
+	UMovieSceneTakeSection* TakeSection = Cast<UMovieSceneTakeSection>(TakeTrack->CreateNewSection());
+	TakeTrack->AddSection(*TakeSection);
+
+	if (FrameRange.IsSet())
+	{
+		TArray<int32> Hours, Minutes, Seconds, Frames;
+		TArray<FMovieSceneFloatValue> SubFrames;
+		TArray<FFrameNumber> Times;
+
+		Hours.Reserve(RecordedTimes.Num());
+		Minutes.Reserve(RecordedTimes.Num());
+		Seconds.Reserve(RecordedTimes.Num());
+		Frames.Reserve(RecordedTimes.Num());
+		SubFrames.Reserve(RecordedTimes.Num());
+		Times.Reserve(RecordedTimes.Num());
+
+		FFrameRate TickResolution = MovieScene->GetTickResolution();
+		FFrameRate DisplayRate = MovieScene->GetDisplayRate();
+
+		FFrameRate TCRate = TickResolution;
+		for (const TPair<FQualifiedFrameTime, FQualifiedFrameTime>& RecordedTimePair : RecordedTimes)
+		{
+			FFrameNumber FrameNumber = RecordedTimePair.Key.Time.FrameNumber;
+			if (!FrameRange.GetValue().Contains(FrameNumber))
+			{
+				continue;
+			}
+
+			FTimecode Timecode = RecordedTimePair.Value.ToTimecode();
+			TCRate = RecordedTimePair.Value.Rate;
+			Hours.Add(Timecode.Hours);
+			Minutes.Add(Timecode.Minutes);
+			Seconds.Add(Timecode.Seconds);
+			Frames.Add(Timecode.Frames);
+
+			FMovieSceneFloatValue SubFrame;
+			if (RecordedTimePair.Value.Time.GetSubFrame() > 0)
+			{
+				// If the Timecode provided gave us a subframe value then we should use that value.  Otherwise, we should compute
+				// the most appropriate value based on the timecode rate.
+				SubFrame.Value = RecordedTimePair.Value.Time.GetSubFrame();
+			}
+			else
+			{
+				FFrameTime FrameTime = FFrameRate::TransformTime(RecordedTimePair.Key.Time, TickResolution, DisplayRate);
+				FQualifiedFrameTime FrameTimeAsTimeCodeRate(FrameTime, TCRate);
+
+				SubFrame.Value = FrameTimeAsTimeCodeRate.Time.GetSubFrame();
+			}
+
+			SubFrame.InterpMode = ERichCurveInterpMode::RCIM_Linear;
+			SubFrames.Add(SubFrame);
+
+			Times.Add(FrameNumber);
+		}
+
+		Hours.Shrink();
+		Minutes.Shrink();
+		Seconds.Shrink();
+		Frames.Shrink();
+		SubFrames.Shrink();
+		Times.Shrink();
+
+		TakeSection->HoursCurve.Set(Times, Hours);
+		TakeSection->MinutesCurve.Set(Times, Minutes);
+		TakeSection->SecondsCurve.Set(Times, Seconds);
+		TakeSection->FramesCurve.Set(Times, Frames);
+		TakeSection->SubFramesCurve.Set(Times, SubFrames);
+		TakeSection->RateCurve.SetDefault(TCRate.AsDecimal());
+	}
+
+	// Since the take section was created post recording here in this
+	// function, it wasn't available at the start of recording to have
+	// its timecode source set with the other sections, so we set it here.
+	if (TakeSection->HoursCurve.GetNumKeys() > 0)
+	{
+		// We populated the take section's timecode curves with data, so
+		// use the first values as the timecode source.
+		const int32 Hours = TakeSection->HoursCurve.GetValues()[0];
+		const int32 Minutes = TakeSection->MinutesCurve.GetValues()[0];
+		const int32 Seconds = TakeSection->SecondsCurve.GetValues()[0];
+		const int32 Frames = TakeSection->FramesCurve.GetValues()[0];
+		const bool bIsDropFrame = false;
+		const FTimecode Timecode(Hours, Minutes, Seconds, Frames, bIsDropFrame);
+		TakeSection->TimecodeSource = FMovieSceneTimecodeSource(Timecode);
+	}
+	else
+	{
+		// Otherwise, adopt the earliest timecode source from one of the movie
+		// scene's other sections as the timecode source for the take section.
+		// This case is unlikely.
+		TakeSection->TimecodeSource = EarliestTimecodeSource;
+	}
+
+	if (UTakeMetaData* TakeMetaData = InSequence->FindMetaData<UTakeMetaData>())
+	{
+		TakeSection->Slate.SetDefault(FString::Printf(TEXT("%s_%d"), *TakeMetaData->GetSlate(), TakeMetaData->GetTakeNumber()));
+	}
+
+	if (TakeSection->GetAutoSizeRange().IsSet())
+	{
+		TakeSection->SetRange(TakeSection->GetAutoSizeRange().GetValue());
+	}
+}
+}
diff --git a/Engine/Plugins/VirtualProduction/Takes/Source/TakesCore/Public/TakeRecorderTimeProcessing.h b/Engine/Plugins/VirtualProduction/Takes/Source/TakesCore/Public/TakeRecorderTimeProcessing.h
new file mode 100644
--- /dev/null
+++ b/Engine/Plugins/VirtualProduction/Takes/Source/TakesCore/Public/TakeRecorderTimeProcessing.h
@@ -0,0 +1,22 @@
+﻿// Copyright Epic Games, Inc. All Rights Reserved.
+
+#pragma once
+
+#include "Containers/Array.h"
+#include "Misc/QualifiedFrameTime.h"
+#include "Templates/Tuple.h"
+
+class ULevelSequence;
+class UMovieSceneTakeTrack;
+
+namespace UE::TakesCore
+{
+/** Key: Recorded frame (to which to apply the timecode). Value: The timecode the recorded frame has. */
+using FArrayOfRecordedTimePairs = TArray<TPair<FQualifiedFrameTime, FQualifiedFrameTime>>;
+
+/** Creates a takes track to store timecode data on a take recorder source. */
+TAKESCORE_API void ProcessRecordedTimes(
+	ULevelSequence* InSequence, UMovieSceneTakeTrack* TakeTrack,
+	const TOptional<TRange<FFrameNumber>>& FrameRange, const FArrayOfRecordedTimePairs& RecordedTimes
+	);
+}
\ No newline at end of file
diff --git a/Engine/Plugins/VirtualProduction/Takes/Source/TakesCore/TakesCore.Build.cs b/Engine/Plugins/VirtualProduction/Takes/Source/TakesCore/TakesCore.Build.cs
--- a/Engine/Plugins/VirtualProduction/Takes/Source/TakesCore/TakesCore.Build.cs
+++ b/Engine/Plugins/VirtualProduction/Takes/Source/TakesCore/TakesCore.Build.cs
@@ -12,7 +12,8 @@ public class TakesCore : ModuleRules
 				"CoreUObject",
 				"EditorFramework",
 				"UnrealEd",
-				"NamingTokens",
+				"NamingTokens", 
+				"TakeMovieScene",
 			}
 		);
 
