Add support for plugin content with take recorder records.

diff --git a/Engine/Plugins/VirtualProduction/Takes/Source/TakeRecorder/Private/TakeRecorderDirectoryHelpers.cpp b/Engine/Plugins/VirtualProduction/Takes/Source/TakeRecorder/Private/TakeRecorderDirectoryHelpers.cpp
--- a/Engine/Plugins/VirtualProduction/Takes/Source/TakeRecorder/Private/TakeRecorderDirectoryHelpers.cpp
+++ b/Engine/Plugins/VirtualProduction/Takes/Source/TakeRecorder/Private/TakeRecorderDirectoryHelpers.cpp
@@ -4,8 +4,8 @@
 
 #include "EditorDirectories.h"
 #include "Misc/PackageName.h"
-
-
+#include "HAL/PlatformFileManager.h"
+#include "HAL/PlatformFile.h"
 namespace UE::TakeRecorder::Private
 {
 
@@ -61,7 +61,18 @@ FString ResolvePathToProject(const FString& InPath)
 	const FString FixedPathStr = InPath.StartsWith(TEXT("/")) ? InPath : TEXT("/") + InPath;
 	const FString DefaultProjectPath = GetDefaultProjectPath();
 
-	return DefaultProjectPath + FixedPathStr;
+	const FString FullProjectPath = DefaultProjectPath + FixedPathStr;
+	const bool bFullProjectPathExists = FPlatformFileManager::Get().GetPlatformFile().DirectoryExists(*FullProjectPath);
+	
+	TTuple<FString, FString> Parts = SplitPackagePath(FixedPathStr);
+	FString PossibleMountPoint = Parts.Get<0>() + TEXT("/");
+	
+	if (!bFullProjectPathExists && FPackageName::MountPointExists(PossibleMountPoint))
+	{
+		return InPath;
+	}
+
+	return FullProjectPath;
 }
 
 FString RemoveProjectFromPath(const FString& InPath)
diff --git a/Engine/Plugins/VirtualProduction/Takes/Source/TakeRecorder/Public/Recorder/TakeRecorderParameters.h b/Engine/Plugins/VirtualProduction/Takes/Source/TakeRecorder/Public/Recorder/TakeRecorderParameters.h
--- a/Engine/Plugins/VirtualProduction/Takes/Source/TakeRecorder/Public/Recorder/TakeRecorderParameters.h
+++ b/Engine/Plugins/VirtualProduction/Takes/Source/TakeRecorder/Public/Recorder/TakeRecorderParameters.h
@@ -71,7 +71,7 @@ struct FTakeRecorderProjectParameters
 	TAKERECORDER_API FString GetTakeAssetPath() const;
 
 	/** The root of the directory in which to save recorded takes. */
-	UPROPERTY(config, EditAnywhere, BlueprintReadWrite, Category = "Take Recorder", meta = (ContentDir, NamingTokens))
+	UPROPERTY(config, EditAnywhere, BlueprintReadWrite, Category = "Take Recorder", meta = (ContentDir, NamingTokens, ForceShowPluginContent))
 	FDirectoryPath RootTakeSaveDir;
 
 	/**
