Take Recorder: Fix recording over an existing level sequence crashing when using the Take Recorder subsystem.

diff --git a/Engine/Plugins/VirtualProduction/Takes/Source/TakeRecorder/Private/Recorder/TakeRecorder.cpp b/Engine/Plugins/VirtualProduction/Takes/Source/TakeRecorder/Private/Recorder/TakeRecorder.cpp
--- a/Engine/Plugins/VirtualProduction/Takes/Source/TakeRecorder/Private/Recorder/TakeRecorder.cpp
+++ b/Engine/Plugins/VirtualProduction/Takes/Source/TakeRecorder/Private/Recorder/TakeRecorder.cpp
@@ -680,8 +680,10 @@ bool UTakeRecorder::SetupDestinationAsset(const FTakeRecorderParameters& InParam
 	}
 
 	// The SequenceAsset is either LevelSequenceBase or the currently focused sequence
-	SequenceAsset = Cast<ULevelSequence>(WeakSequencer.Pin()->GetFocusedMovieSceneSequence());
-
+	SequenceAsset = WeakSequencer.IsValid() ? Cast<ULevelSequence>(WeakSequencer.Pin()->GetFocusedMovieSceneSequence())
+		: LevelSequenceBase;
+	check(SequenceAsset);
+	
 	// Copy the sources into the level sequence for future reference (and potentially mutation throughout recording)
 	SequenceAsset->CopyMetaData(Sources);
 
