Take Recorder: Save per-frame hitch information into ULevelSequence. The data is saved in a Sequencer decorator object, which is not displayed to the user.

The per-frame data consists of two FTimecode "tracks": the target timecode, which is the timecode the frame was supposed to have, and the actual timecode, which is the frame that the linear regression system predicted for the frame. When they don't match, that means  that either the engine has skipped a timecode frame or has hitched.

diff --git a/Engine/Plugins/VirtualProduction/Takes/Source/TakeMovieScene/Private/Hitching/FrameHitchSceneDecoration.cpp b/Engine/Plugins/VirtualProduction/Takes/Source/TakeMovieScene/Private/Hitching/FrameHitchSceneDecoration.cpp
new file mode 100644
--- /dev/null
+++ b/Engine/Plugins/VirtualProduction/Takes/Source/TakeMovieScene/Private/Hitching/FrameHitchSceneDecoration.cpp
@@ -0,0 +1,10 @@
+﻿// Copyright Epic Games, Inc. All Rights Reserved.
+
+#include "Hitching/FrameHitchSceneDecoration.h"
+
+TOptional<UE::TakeMovieScene::FFrameHitchData> UFrameHitchSceneDecoration::Evaluate(const FFrameTime& InFrameTime) const
+{
+	const TOptional<FTimecode> Target = TargetTimecode.Evaluate(InFrameTime);
+	const TOptional<FTimecode> Actual = ActualTimecode.Evaluate(InFrameTime);
+	return Target && Actual ? UE::TakeMovieScene::FFrameHitchData{ *Target, *Actual } : TOptional<UE::TakeMovieScene::FFrameHitchData>();
+}
diff --git a/Engine/Plugins/VirtualProduction/Takes/Source/TakeMovieScene/Private/Hitching/TimecodeHitchChannels.cpp b/Engine/Plugins/VirtualProduction/Takes/Source/TakeMovieScene/Private/Hitching/TimecodeHitchChannels.cpp
new file mode 100644
--- /dev/null
+++ b/Engine/Plugins/VirtualProduction/Takes/Source/TakeMovieScene/Private/Hitching/TimecodeHitchChannels.cpp
@@ -0,0 +1,40 @@
+﻿// Copyright Epic Games, Inc. All Rights Reserved.
+
+#include "Hitching/TimecodeHitchChannels.h"
+
+void FTakeMovieSceneHitchTimecodeCurves::Set(
+	TArray<FFrameNumber> Times, TArray<int32> Hours, TArray<int32> Minutes, TArray<int32> Seconds, TArray<int32> Frames
+	)
+{
+	const int32 Num = Times.Num();
+	if (ensure(Hours.Num() == Num && Minutes.Num() == Num && Seconds.Num() == Num && Frames.Num() == Num))
+	{
+		HoursCurve.Set(Times, MoveTemp(Hours));
+		MinutesCurve.Set(Times, MoveTemp(Minutes));
+		SecondsCurve.Set(Times, MoveTemp(Seconds));
+		FramesCurve.Set(MoveTemp(Times), MoveTemp(Frames));
+	}
+}
+
+TOptional<FTimecode> FTakeMovieSceneHitchTimecodeCurves::Evaluate(const FFrameTime& InFrameTime) const
+{
+	int32 Hours = 0;
+	int32 Minutes = 0;
+	int32 Seconds = 0;
+	int32 Frames = 0;
+
+	if (!HoursCurve.Evaluate(InFrameTime, Hours)
+		|| !MinutesCurve.Evaluate(InFrameTime, Minutes)
+		|| !SecondsCurve.Evaluate(InFrameTime, Seconds)
+		|| !FramesCurve.Evaluate(InFrameTime, Frames))
+	{
+		return {};
+	}
+
+	return FTimecode(Hours, Minutes, Seconds, Frames, false, false);
+}
+
+TConstArrayView<FFrameNumber> FTakeMovieSceneHitchTimecodeCurves::GetFrameTimes() const
+{
+	return HoursCurve.GetTimes();
+}
diff --git a/Engine/Plugins/VirtualProduction/Takes/Source/TakeMovieScene/Public/Hitching/FrameHitchData.h b/Engine/Plugins/VirtualProduction/Takes/Source/TakeMovieScene/Public/Hitching/FrameHitchData.h
new file mode 100644
--- /dev/null
+++ b/Engine/Plugins/VirtualProduction/Takes/Source/TakeMovieScene/Public/Hitching/FrameHitchData.h
@@ -0,0 +1,16 @@
+﻿// Copyright Epic Games, Inc. All Rights Reserved.
+
+#pragma once
+
+#include "Misc/Timecode.h"
+
+namespace UE::TakeMovieScene
+{
+struct FFrameHitchData
+{
+	/** The timecode frame engine was supposed to have. */
+	FTimecode TargetTimecode;
+	/** The timecode the frame actually had. */
+	FTimecode ActualTimecode;
+};
+}
diff --git a/Engine/Plugins/VirtualProduction/Takes/Source/TakeMovieScene/Public/Hitching/FrameHitchSceneDecoration.h b/Engine/Plugins/VirtualProduction/Takes/Source/TakeMovieScene/Public/Hitching/FrameHitchSceneDecoration.h
new file mode 100644
--- /dev/null
+++ b/Engine/Plugins/VirtualProduction/Takes/Source/TakeMovieScene/Public/Hitching/FrameHitchSceneDecoration.h
@@ -0,0 +1,35 @@
+﻿// Copyright Epic Games, Inc. All Rights Reserved.
+
+#pragma once
+
+#include "FrameHitchData.h"
+#include "TimecodeHitchChannels.h"
+#include "UObject/Object.h"
+#include "FrameHitchSceneDecoration.generated.h"
+
+/**
+ * This decoration contain stores a target and actual FTimecode for each frame for the purposes of detecting whether the engine was running behind,
+ * i.e. was hitching, on a previous frame during take recording.
+ *
+ * This data can be used to visualize hitches in the Sequencer UI.
+ */
+UCLASS()
+class UFrameHitchSceneDecoration : public UObject
+{
+	GENERATED_BODY()
+public:
+
+	/** Target timecode that the frame was supposed to have */
+	UPROPERTY()
+	FTakeMovieSceneHitchTimecodeCurves TargetTimecode;
+	/** Actual timecode that was recorded */
+	UPROPERTY()
+	FTakeMovieSceneHitchTimecodeCurves ActualTimecode;
+
+	/** The frame rate at which was recorded. Determines the number of frames per second in the FTimecode. */
+	UPROPERTY()
+	FFrameRate FrameRate;
+
+	/** Evaluates the curve data at the specified time. */
+	TAKEMOVIESCENE_API TOptional<UE::TakeMovieScene::FFrameHitchData> Evaluate(const FFrameTime& InFrameTime) const;
+};
diff --git a/Engine/Plugins/VirtualProduction/Takes/Source/TakeMovieScene/Public/Hitching/TimecodeHitchChannels.h b/Engine/Plugins/VirtualProduction/Takes/Source/TakeMovieScene/Public/Hitching/TimecodeHitchChannels.h
new file mode 100644
--- /dev/null
+++ b/Engine/Plugins/VirtualProduction/Takes/Source/TakeMovieScene/Public/Hitching/TimecodeHitchChannels.h
@@ -0,0 +1,36 @@
+﻿// Copyright Epic Games, Inc. All Rights Reserved.
+
+#pragma once
+
+#include "Channels/MovieSceneIntegerChannel.h"
+#include "Misc/Optional.h"
+#include "Misc/Timecode.h"
+#include "TimecodeHitchChannels.generated.h"
+
+USTRUCT()
+struct FTakeMovieSceneHitchTimecodeCurves
+{
+	GENERATED_BODY()
+
+	/** Hours curve data  */
+	UPROPERTY()
+	FMovieSceneIntegerChannel HoursCurve;
+	/** Minutes curve data */
+	UPROPERTY()
+	FMovieSceneIntegerChannel MinutesCurve;
+	/** Seconds curve data */
+	UPROPERTY()
+	FMovieSceneIntegerChannel SecondsCurve;
+	/** Frames curve data */
+	UPROPERTY()
+	FMovieSceneIntegerChannel FramesCurve;
+
+	/** Sets the curve content from the given arrays - which must be equal size. */
+	TAKEMOVIESCENE_API void Set(TArray<FFrameNumber> Times, TArray<int32> Hours, TArray<int32> Minutes, TArray<int32> Seconds, TArray<int32> Frames);
+
+	/** @return The timecode at the frame time. */
+	TAKEMOVIESCENE_API TOptional<FTimecode> Evaluate(const FFrameTime& InFrameTime) const;
+
+	/** @return Gets all the frame times at which we recorded. */
+	TAKEMOVIESCENE_API TConstArrayView<FFrameNumber> GetFrameTimes() const;
+};
\ No newline at end of file
diff --git a/Engine/Plugins/VirtualProduction/Takes/Source/TakeRecorder/Private/Recorder/TakeRecorder.cpp b/Engine/Plugins/VirtualProduction/Takes/Source/TakeRecorder/Private/Recorder/TakeRecorder.cpp
--- a/Engine/Plugins/VirtualProduction/Takes/Source/TakeRecorder/Private/Recorder/TakeRecorder.cpp
+++ b/Engine/Plugins/VirtualProduction/Takes/Source/TakeRecorder/Private/Recorder/TakeRecorder.cpp
@@ -1041,6 +1041,7 @@ void UTakeRecorder::Start()
 
 		// Record immediately so that there's a key on the first frame of recording
 		Sources->TickRecording(SequenceAsset, RecordTime, 0.1f);
+		OnTickRecordingEvent.Broadcast(this, RecordTime);
 	}
 
 	if (!ShouldShowNotifications())
diff --git a/Engine/Plugins/VirtualProduction/Takes/Source/TakeRecorder/Private/TakeRecorderModule.cpp b/Engine/Plugins/VirtualProduction/Takes/Source/TakeRecorder/Private/TakeRecorderModule.cpp
--- a/Engine/Plugins/VirtualProduction/Takes/Source/TakeRecorder/Private/TakeRecorderModule.cpp
+++ b/Engine/Plugins/VirtualProduction/Takes/Source/TakeRecorder/Private/TakeRecorderModule.cpp
@@ -334,7 +334,7 @@ void FTakeRecorderModule::StartupModule()
 	RegisterLevelEditorExtensions();
 	RegisterSettings();
 	RegisterSerializedRecorder();
-	HitchlessTimecodeIntegration = MakeUnique<UE::TakeRecorder::FTimecodeRegressionRecordSetup>();
+	TimecodeManagement = MakeUnique<UE::TakeRecorder::FHitchlessProtectionRootLogic>();
 
 #if WITH_EDITOR
 	if (GIsEditor)
@@ -374,7 +374,7 @@ void FTakeRecorderModule::ShutdownModule()
 	UnregisterLevelEditorExtensions();
 	UnregisterSettings();
 	UnregisterSerializedRecorder();
-	HitchlessTimecodeIntegration.Reset();
+	TimecodeManagement.Reset();
 }
 
 UTakePreset* FTakeRecorderModule::GetPendingTake() const
diff --git a/Engine/Plugins/VirtualProduction/Takes/Source/TakeRecorder/Private/TakeRecorderModule.h b/Engine/Plugins/VirtualProduction/Takes/Source/TakeRecorder/Private/TakeRecorderModule.h
--- a/Engine/Plugins/VirtualProduction/Takes/Source/TakeRecorder/Private/TakeRecorderModule.h
+++ b/Engine/Plugins/VirtualProduction/Takes/Source/TakeRecorder/Private/TakeRecorderModule.h
@@ -6,9 +6,8 @@
 #include "Delegates/IDelegateInstance.h"
 #include "UObject/GCObject.h"
 #include "ITakeRecorderModule.h"
+#include "Timecode/HitchlessProtectionRootLogic.h"
 #include "UObject/WeakObjectPtrTemplates.h"
-#include "Timecode/Regression/TimecodeRegressionRecordSetup.h"
-#include "Widgets/Input/SCheckBox.h"
 
 class FExtender;
 class FTakePresetActions;
@@ -94,13 +93,13 @@ private:
 	TSharedPtr<FSerializedRecorder> SerializedRecorder;
 
 	/**
-	 * Sets up the UEngine::TimecodeProvider so timecode is not skipped when the engine hitches.
+	 * Sets up all sub-systems related to managing timecode in Take Recorder
 	 * 
 	 * For future maintainers: this could totally live in a separate module / plugin / whatever.
 	 * There's no technical reason to have this live in TakeRecorder. The reason it's here is that we probably always want hitching to be handled
 	 * gracefully.
 	 */
-	TUniquePtr<UE::TakeRecorder::FTimecodeRegressionRecordSetup> HitchlessTimecodeIntegration;
+	TUniquePtr<UE::TakeRecorder::FHitchlessProtectionRootLogic> TimecodeManagement;
 
 	/** Extra class default objects that should appear in the Take Recorder panel. */
 	TArray<TWeakObjectPtr<>> ExternalObjects;
diff --git a/Engine/Plugins/VirtualProduction/Takes/Source/TakeRecorder/Private/Timecode/HitchlessProtectionRootLogic.cpp b/Engine/Plugins/VirtualProduction/Takes/Source/TakeRecorder/Private/Timecode/HitchlessProtectionRootLogic.cpp
new file mode 100644
--- /dev/null
+++ b/Engine/Plugins/VirtualProduction/Takes/Source/TakeRecorder/Private/Timecode/HitchlessProtectionRootLogic.cpp
@@ -0,0 +1,97 @@
+﻿// Copyright Epic Games, Inc. All Rights Reserved.
+
+#include "HitchlessProtectionRootLogic.h"
+
+#include "Recorder/TakeRecorderSubsystem.h"
+
+namespace UE::TakeRecorder
+{
+namespace Private
+{
+static TAutoConsoleVariable<bool> CVarEnableHitchlessTimecodeProvider(
+	TEXT("TakeRecorder.EnableHitchlessTimecode"),
+	false,
+	TEXT("If true, we'll try to handle engine hitching. This requires that the UEngineCustomTimestep implements the IClockedTimeStep interface.")
+	);
+static TAutoConsoleVariable<float> CVarTimecodeLinearRegressionSeconds(
+	TEXT("TakeRecorder.TimecodeLinearRegressionSeconds"),
+	1.0f,
+	TEXT("When using linear regression to predict timecode, how many seconds worth of frames to sample.")
+	);
+
+static TAutoConsoleVariable<bool> CVarTestWithSource(
+	TEXT("Dominik.TestWithSource"),
+	true,
+	TEXT("When using linear regression to predict timecode, how many seconds worth of frames to sample.")
+);
+}
+	
+FHitchlessProtectionRootLogic::FHitchlessProtectionRootLogic()
+{
+	if (IsEngineStartupModuleLoadingComplete())
+	{
+		SetupTakeRecorderDelegates();
+	}
+	else
+	{
+		FCoreDelegates::OnPostEngineInit.AddRaw(this, &FHitchlessProtectionRootLogic::SetupTakeRecorderDelegates);
+	}
+}
+
+FHitchlessProtectionRootLogic::~FHitchlessProtectionRootLogic()
+{
+	FCoreDelegates::OnPostEngineInit.RemoveAll(this);
+
+	if (UTakeRecorderSubsystem* TakeRecorderSubsystem = GEngine ? GEngine->GetEngineSubsystem<UTakeRecorderSubsystem>() : nullptr)
+	{
+		TakeRecorderSubsystem->GetOnRecordingInitializedEvent().RemoveAll(this);
+		TakeRecorderSubsystem->GetOnRecordingFinishedEvent().RemoveAll(this);
+		TakeRecorderSubsystem->GetOnRecordingCancelledEvent().RemoveAll(this);
+	}
+}
+
+void FHitchlessProtectionRootLogic::SetupTakeRecorderDelegates()
+{
+	UTakeRecorderSubsystem* TakeRecorderSubsystem = GEngine ? GEngine->GetEngineSubsystem<UTakeRecorderSubsystem>() : nullptr;
+	if (ensure(TakeRecorderSubsystem))
+	{
+		TakeRecorderSubsystem->GetOnRecordingInitializedEvent().AddRaw(this, &FHitchlessProtectionRootLogic::OnInitializeRecording);
+		TakeRecorderSubsystem->GetOnRecordingFinishedEvent().AddRaw(this, &FHitchlessProtectionRootLogic::OnFinishRecording);
+		TakeRecorderSubsystem->GetOnRecordingCancelledEvent().AddRaw(this, &FHitchlessProtectionRootLogic::OnCancelRecording);
+	}
+}
+
+void FHitchlessProtectionRootLogic::OnInitializeRecording(UTakeRecorder* TakeRecorder)
+{
+	if (!Private::CVarEnableHitchlessTimecodeProvider.GetValueOnAnyThread())
+	{
+		return;
+	}
+
+	const TOptional<TNonNullPtr<UTimecodeProvider>> UnderlyingTimecodeProvider = TimecodeRegressionRecording.SetupEngineEnvironmentForRecording(
+		Private::CVarTimecodeLinearRegressionSeconds.GetValueOnAnyThread()
+		);
+	if (!UnderlyingTimecodeProvider)
+	{
+		return;
+	}
+	
+	RecordingState.Emplace(*UnderlyingTimecodeProvider, TakeRecorder);
+}
+
+void FHitchlessProtectionRootLogic::OnFinishRecording(UTakeRecorder* TakeRecorder)
+{
+	CleanupRecording();
+}
+
+void FHitchlessProtectionRootLogic::OnCancelRecording(UTakeRecorder* TakeRecorder)
+{
+	CleanupRecording(); 
+}
+
+void FHitchlessProtectionRootLogic::CleanupRecording()
+{
+	RecordingState.Reset();
+	TimecodeRegressionRecording.CleanupRecording();
+}
+}
diff --git a/Engine/Plugins/VirtualProduction/Takes/Source/TakeRecorder/Private/Timecode/HitchlessProtectionRootLogic.h b/Engine/Plugins/VirtualProduction/Takes/Source/TakeRecorder/Private/Timecode/HitchlessProtectionRootLogic.h
new file mode 100644
--- /dev/null
+++ b/Engine/Plugins/VirtualProduction/Takes/Source/TakeRecorder/Private/Timecode/HitchlessProtectionRootLogic.h
@@ -0,0 +1,59 @@
+﻿// Copyright Epic Games, Inc. All Rights Reserved.
+
+#pragma once
+
+#include "Recording/FrameHitchDecorationRecorder.h"
+#include "Regression/TimecodeRegressionRecordSetup.h"
+#include "Templates/UnrealTemplate.h"
+
+namespace UE::TakeRecorder
+{
+/**
+ * Sets up all systems related to managing timecode in Take Recorder and dealing with engine hitches, such as timecode correction during recording and
+ * visualization of hitches in Sequencer.
+ * 
+ * This is the root object class that handles injecting all dependencies into the subsystems. It acts a mediator between the systems ensuring the
+ * correct order of operations during recording & reviewing a take.
+ */
+class FHitchlessProtectionRootLogic : public FNoncopyable
+{
+public:
+
+	FHitchlessProtectionRootLogic();
+	~FHitchlessProtectionRootLogic();
+
+private:
+
+	/**
+	 * Overrides the global FApp's timecode by estimating the timecode the frame has using linear regression.
+	 * 
+	 * E.g. when a 1s hitch occurs at time 14:11:06:01, then the next frame should be 14:11:06:02 even though the real timecode is 14:11:07:02 meaning
+	 * that the engine has to catch up. This sets up the relevant systems.
+	 */
+	FTimecodeRegressionRecordSetup TimecodeRegressionRecording;
+	
+	struct FRecordingState
+	{
+		/** Records the underlying timecode provider and the one that is estimated by UTimecodeRegressionProvider. */
+		FFrameHitchDecorationRecorder TimecodeRecording;
+
+		explicit FRecordingState(TNotNull<UTimecodeProvider*> InRecordedProvider, TNotNull<UTakeRecorder*> InTakeRecorder)
+			: TimecodeRecording(InRecordedProvider, InTakeRecorder)
+		{}
+	};
+	/** Set while recording. */
+	TOptional<FRecordingState> RecordingState;
+
+	/** Starts listening for UTakeRecorder events. */
+	void SetupTakeRecorderDelegates();
+	/** Sets up hitchless recording. */
+	void OnInitializeRecording(UTakeRecorder* TakeRecorder);
+	/** Creates timecode tracks requireds for visualization of timecode hitches. */
+	void OnFinishRecording(UTakeRecorder* TakeRecorder);
+	/** Clears any state */
+	void OnCancelRecording(UTakeRecorder* TakeRecorder);
+
+	void CleanupRecording();
+};
+}
+
diff --git a/Engine/Plugins/VirtualProduction/Takes/Source/TakeRecorder/Private/Timecode/Recording/FrameHitchDecorationRecorder.cpp b/Engine/Plugins/VirtualProduction/Takes/Source/TakeRecorder/Private/Timecode/Recording/FrameHitchDecorationRecorder.cpp
new file mode 100644
--- /dev/null
+++ b/Engine/Plugins/VirtualProduction/Takes/Source/TakeRecorder/Private/Timecode/Recording/FrameHitchDecorationRecorder.cpp
@@ -0,0 +1,114 @@
+﻿// Copyright Epic Games, Inc. All Rights Reserved.
+
+#include "FrameHitchDecorationRecorder.h"
+
+#include "Engine/TimecodeProvider.h"
+#include "Hitching/FrameHitchSceneDecoration.h"
+#include "LevelSequence.h"
+#include "MovieScene.h"
+#include "Recorder/TakeRecorder.h"
+#include "TakeRecorderSources.h"
+
+#define LOCTEXT_NAMESPACE "FTimecodeTrackRecording"
+
+namespace UE::TakeRecorder
+{
+namespace Private
+{
+/** Recurse all sections and get the maximum range used. */
+static void CombineRange(TOptional<TRange<FFrameNumber>>& InOutRange, TConstArrayView<UMovieSceneTrack*> Tracks)
+{
+	for (UMovieSceneTrack* Track : Tracks)
+	{
+		for (UMovieSceneSection* Section : Track->GetAllSections())
+		{
+			if (Section->HasStartFrame() && Section->HasEndFrame())
+			{
+				if (!InOutRange.IsSet())
+				{
+					InOutRange = Section->GetRange();
+				}
+				else
+				{
+					InOutRange = TRange<FFrameNumber>::Hull(InOutRange.GetValue(), Section->GetRange());
+				}
+			}
+		}
+	}
+}
+
+/** Recurse sections and get the maximum range used. */
+static TOptional<TRange<FFrameNumber>> GetGlobalFrameRange(const UMovieScene* MovieScene)
+{
+	TOptional<TRange<FFrameNumber>> FrameRange;
+
+	CombineRange(FrameRange, MovieScene->GetTracks());
+	for (const FMovieSceneBinding& Binding : MovieScene->GetBindings())
+	{
+		CombineRange(FrameRange, Binding.GetTracks());
+	}
+
+	return FrameRange;
+}
+}
+FFrameHitchDecorationRecorder::FFrameHitchDecorationRecorder(TNotNull<UTimecodeProvider*> InRecordedProvider, TNotNull<UTakeRecorder*> InTakeRecorder)
+	: WeakTimecodeProvider(InRecordedProvider)
+	, WeakTakeRecorder(InTakeRecorder)
+{
+	InTakeRecorder->OnTickRecording().AddRaw(this, &FFrameHitchDecorationRecorder::OnTickRecording);
+	InTakeRecorder->OnRecordingStopped().AddRaw(this, &FFrameHitchDecorationRecorder::OnRecordingStopped);
+}
+
+FFrameHitchDecorationRecorder::~FFrameHitchDecorationRecorder()
+{
+	if (UTakeRecorder* TakeRecorder = WeakTakeRecorder.Get())
+	{
+		TakeRecorder->OnTickRecording().RemoveAll(this);
+		TakeRecorder->OnRecordingStopped().RemoveAll(this);
+	}
+}
+
+void FFrameHitchDecorationRecorder::OnTickRecording(UTakeRecorder* InTakeRecorder, const FQualifiedFrameTime& InCurrentFrameTime)
+{
+	if (UTimecodeProvider* TimecodeProvider = WeakTimecodeProvider.Get())
+	{
+		const FQualifiedFrameTime RealTimeCode = TimecodeProvider->GetQualifiedFrameTime();
+		TargetTimecodeValues.Emplace(InCurrentFrameTime, RealTimeCode);
+	}
+}
+
+void FFrameHitchDecorationRecorder::OnRecordingStopped(UTakeRecorder* TakeRecorder) const
+{
+	using namespace TakesCore;
+	ULevelSequence* RootSequence = TakeRecorder->GetSequence();
+	UMovieScene* MovieScene = RootSequence->GetMovieScene();
+
+	const TOptional<TRange<FFrameNumber>> GlobalFrameRange = Private::GetGlobalFrameRange(MovieScene);
+	if (!GlobalFrameRange)
+	{
+		return;
+	}
+	
+	UFrameHitchSceneDecoration* Decoration = MovieScene->GetOrCreateDecoration<UFrameHitchSceneDecoration>();
+	const ETimecodeExtractionFlags Flags = ETimecodeExtractionFlags::Hours | ETimecodeExtractionFlags::Minutes
+		| ETimecodeExtractionFlags::Seconds | ETimecodeExtractionFlags::Frames | ETimecodeExtractionFlags::Times;
+	
+	FTimecodeExtractionData TargetData = ExtractIntoTimecodeArrays(MovieScene, *GlobalFrameRange, TargetTimecodeValues, Flags);
+	FTimecodeExtractionData ActualData = ExtractIntoTimecodeArrays(MovieScene, *GlobalFrameRange, UTakeRecorderSources::RecordedTimes, Flags);
+	ensureMsgf(TargetData.Times.Num() == ActualData.Times.Num(),
+		TEXT("Number of frames tagged with target timecode is %d but frames with actual timecode is %d. Investigate why they don't coincide."),
+		TargetData.Times.Num(), ActualData.Times.Num()
+		);
+	ensure(TargetData.TCRate == ActualData.TCRate);
+	
+	Decoration->TargetTimecode.Set(
+		MoveTemp(TargetData.Times), MoveTemp(TargetData.Hours), MoveTemp(TargetData.Minutes), MoveTemp(TargetData.Seconds), MoveTemp(TargetData.Frames)
+		);
+	Decoration->ActualTimecode.Set(
+		MoveTemp(ActualData.Times), MoveTemp(ActualData.Hours), MoveTemp(ActualData.Minutes), MoveTemp(ActualData.Seconds), MoveTemp(ActualData.Frames)
+		);
+	Decoration->FrameRate = TargetData.TCRate; 
+}
+}
+
+#undef LOCTEXT_NAMESPACE
\ No newline at end of file
diff --git a/Engine/Plugins/VirtualProduction/Takes/Source/TakeRecorder/Private/Timecode/Recording/FrameHitchDecorationRecorder.h b/Engine/Plugins/VirtualProduction/Takes/Source/TakeRecorder/Private/Timecode/Recording/FrameHitchDecorationRecorder.h
new file mode 100644
--- /dev/null
+++ b/Engine/Plugins/VirtualProduction/Takes/Source/TakeRecorder/Private/Timecode/Recording/FrameHitchDecorationRecorder.h
@@ -0,0 +1,46 @@
+﻿// Copyright Epic Games, Inc. All Rights Reserved.
+
+#pragma once
+
+#include "TakeRecorderTimeProcessing.h"
+#include "Misc/NotNull.h"
+#include "Misc/QualifiedFrameTime.h"
+#include "Templates/UnrealTemplate.h"
+#include "UObject/WeakObjectPtrTemplates.h"
+
+class UTimecodeProvider;
+class UTakeRecorder;
+
+namespace UE::TakeRecorder
+{
+/**
+ * Creates 2 timecode tracks at the root Sequence:
+ *  - UOriginalTimecodeMovieSceneTrack: Holds the unestimated timecode of the UTimecodeProvider underlying UTimecodeRegressionProvider.
+ *  - URecordedTimecodeMovieSceneTrack: Holds the timecode of FApp::CurrentFrameTime.  */
+class FFrameHitchDecorationRecorder : public FNoncopyable
+{
+public:
+	
+	explicit FFrameHitchDecorationRecorder(TNotNull<UTimecodeProvider*> InRecordedProvider, TNotNull<UTakeRecorder*> InTakeRecorder);
+	~FFrameHitchDecorationRecorder();
+
+private:
+
+	/** The timecode provider being recorded. */
+	const TWeakObjectPtr<UTimecodeProvider> WeakTimecodeProvider;
+
+	/** Take recorder for which we are recording. */
+	const TWeakObjectPtr<UTakeRecorder> WeakTakeRecorder;
+	
+	/** Passed in to TakeRecorderSourceHelpers::ProcessRecordedTimes to generate the track for original timecode. */
+	TakesCore::FArrayOfRecordedTimePairs TargetTimecodeValues;
+
+	/** Called when InTakeRecorder ticks during recording. */
+	void OnTickRecording(UTakeRecorder* InTakeRecorder, const FQualifiedFrameTime& InCurrentFrameTime);
+	
+	/** Creates the tracks and adds them to UTakeRecorder. */
+	void OnRecordingStopped(UTakeRecorder* TakeRecorder) const;
+};
+}
+
+
diff --git a/Engine/Plugins/VirtualProduction/Takes/Source/TakeRecorder/Private/Timecode/Regression/TimecodeRegressionRecordSetup.cpp b/Engine/Plugins/VirtualProduction/Takes/Source/TakeRecorder/Private/Timecode/Regression/TimecodeRegressionRecordSetup.cpp
--- a/Engine/Plugins/VirtualProduction/Takes/Source/TakeRecorder/Private/Timecode/Regression/TimecodeRegressionRecordSetup.cpp
+++ b/Engine/Plugins/VirtualProduction/Takes/Source/TakeRecorder/Private/Timecode/Regression/TimecodeRegressionRecordSetup.cpp
@@ -17,19 +17,8 @@ namespace UE::TakeRecorder
 {
 namespace Private
 {
-static TAutoConsoleVariable<bool> CVarEnableHitchlessTimecodeProvider(
-	TEXT("TakeRecorder.EnableHitchlessTimecode"),
-	false,
-	TEXT("If true, we'll try to handle engine hitching. This requires that the UEngineCustomTimestep implements the IClockedTimeStep interface.")
-	);
-static TAutoConsoleVariable<float> CVarTimecodeLinearRegressionSeconds(
-	TEXT("TakeRecorder.TimecodeLinearRegressionSeconds"),
-	1.0f,
-	TEXT("When using linear regression to predict timecode, how many seconds worth of frames to sample.")
-	);
-
 /** @return The timecode estimator instance we'll use to estimate timecode - effectively handling any engine hitches. */
-static UTimecodeRegressionProvider* CreateTimecodeEstimator(UTimecodeProvider* InActualImplementation)
+static UTimecodeRegressionProvider* CreateTimecodeEstimator(UTimecodeProvider* InActualImplementation, float InRegressionFramesInSeconds)
 {
 	UTimecodeRegressionProvider* TimecodeRegression = NewObject<UTimecodeRegressionProvider>(GetTransientPackage());
 	TimecodeRegression->SetTimecodeImplementation(InActualImplementation);
@@ -38,79 +27,38 @@ static UTimecodeRegressionProvider* CreateTimecodeEstimator(UTimecodeProvider* I
 	// E.g. if FrameRate is 60 FPS of 60 frames, then Numerator = 60 meaning 60 / 1 = 60. We want
 	const FFrameRate FrameRate = GEngine->GetEngineSubsystem<UTakeRecorderSubsystem>()->GetFrameRate();
 	TimecodeRegression->NumSampledFrames = FrameRate.Denominator == 0
-		? 1 : CVarTimecodeLinearRegressionSeconds.GetValueOnAnyThread() * FrameRate.Numerator / FrameRate.Denominator;
+		? 1 : InRegressionFramesInSeconds * FrameRate.Numerator / FrameRate.Denominator;
 	
 	return TimecodeRegression;
 }
 }
-	
-FTimecodeRegressionRecordSetup::FTimecodeRegressionRecordSetup()
-{
-	if (IsEngineStartupModuleLoadingComplete())
-	{
-		OnPostEngineInit();
-	}
-	else
-	{
-		FCoreDelegates::OnPostEngineInit.AddRaw(this, &FTimecodeRegressionRecordSetup::OnPostEngineInit);
-	}
-}
-
-FTimecodeRegressionRecordSetup::~FTimecodeRegressionRecordSetup()
-{
-	FCoreDelegates::OnPostEngineInit.RemoveAll(this);
-	
-	CleanUp();
-}
 
-void FTimecodeRegressionRecordSetup::OnPostEngineInit()
-{
-	UTakeRecorderSubsystem* TakeRecorderSubsystem = GEngine ? GEngine->GetEngineSubsystem<UTakeRecorderSubsystem>() : nullptr;
-	if (ensure(TakeRecorderSubsystem))
-	{
-		TakeRecorderSubsystem->GetOnRecordingInitializedEvent().AddRaw(this, &FTimecodeRegressionRecordSetup::OnInitializeRecording);
-		TakeRecorderSubsystem->GetOnRecordingFinishedEvent().AddRaw(this, &FTimecodeRegressionRecordSetup::OnFinishRecording);
-		TakeRecorderSubsystem->GetOnRecordingCancelledEvent().AddRaw(this, &FTimecodeRegressionRecordSetup::OnCancelRecording);
-	}
-}
-
-void FTimecodeRegressionRecordSetup::CleanUp()
-{
-	if (UTakeRecorderSubsystem* TakeRecorderSubsystem = GEngine ? GEngine->GetEngineSubsystem<UTakeRecorderSubsystem>() : nullptr)
-	{
-		TakeRecorderSubsystem->GetOnRecordingInitializedEvent().RemoveAll(this);
-		TakeRecorderSubsystem->GetOnRecordingFinishedEvent().RemoveAll(this);
-		TakeRecorderSubsystem->GetOnRecordingCancelledEvent().RemoveAll(this);
-		
-		RecordData.Reset();
-	}
-}
-
-void FTimecodeRegressionRecordSetup::SetupEngineEnvironmentForRecording()
+TOptional<TNonNullPtr<UTimecodeProvider>> FTimecodeRegressionRecordSetup::SetupEngineEnvironmentForRecording(float InRegressionFramesInSeconds)
 {
 	UTimecodeProvider* CurrentTimecodeProvider = GEngine ? GEngine->GetTimecodeProvider() : nullptr;
-	if (!Private::CVarEnableHitchlessTimecodeProvider.GetValueOnAnyThread()
-		|| !CurrentTimecodeProvider)
+	if (!CurrentTimecodeProvider)
 	{
-		return;
+		return {};
 	}
 
 	UEngineCustomTimeStep* PreexistingTimestep = GEngine->GetCustomTimeStep();
 	if (PreexistingTimestep && !UTimecodeRegressionProvider::SupportsTimestep(PreexistingTimestep))
 	{
 		ShowIncompatibleTimestepNotification(PreexistingTimestep);
-		return;
+		return {};
 	}
 
 	if (PreexistingTimestep)
 	{
 		// Use pre-configured custom timestep
-		RecordData.Emplace(Private::CreateTimecodeEstimator(CurrentTimecodeProvider));
+		RecordData.Emplace(Private::CreateTimecodeEstimator(CurrentTimecodeProvider, InRegressionFramesInSeconds));
 	}
 	else if (UEngineCustomTimeStep* DefaultTimestep = AskUserToCreateFallbackTimestep())
 	{
 		// There was no pre-configured custom timestep... clean it up when recording finishes.
-		RecordData.Emplace(DefaultTimestep, Private::CreateTimecodeEstimator(CurrentTimecodeProvider));
+		RecordData.Emplace(DefaultTimestep, Private::CreateTimecodeEstimator(CurrentTimecodeProvider, InRegressionFramesInSeconds));
 	}
+
+	return TNonNullPtr<UTimecodeProvider>(CurrentTimecodeProvider);
 }
 }
\ No newline at end of file
diff --git a/Engine/Plugins/VirtualProduction/Takes/Source/TakeRecorder/Private/Timecode/Regression/TimecodeRegressionRecordSetup.h b/Engine/Plugins/VirtualProduction/Takes/Source/TakeRecorder/Private/Timecode/Regression/TimecodeRegressionRecordSetup.h
--- a/Engine/Plugins/VirtualProduction/Takes/Source/TakeRecorder/Private/Timecode/Regression/TimecodeRegressionRecordSetup.h
+++ b/Engine/Plugins/VirtualProduction/Takes/Source/TakeRecorder/Private/Timecode/Regression/TimecodeRegressionRecordSetup.h
@@ -19,8 +19,16 @@ class FTimecodeRegressionRecordSetup : public FNoncopyable
 {
 public:
 
-	FTimecodeRegressionRecordSetup();
-	~FTimecodeRegressionRecordSetup();
+	/**
+	 * Sets up timecode and optionally the custom time step.
+	 * @param InRegressionFramesInSeconds Determines size of the linear regression buffer. This is the number of seconds worth of frames we should use.
+	 *	Example: 1s worth of frames means a buffer size of 24 timecode samples if the UTakeRecorder is set to record at 24 FPS.
+	 * @return The real timecode provider that we will be estimating. Unset if initialization failed.
+	 */
+	TOptional<TNonNullPtr<UTimecodeProvider>> SetupEngineEnvironmentForRecording(float InRegressionFramesInSeconds);
+
+	/** Clears the time step and timecode provider overrides. */
+	void CleanupRecording() { RecordData.Reset(); }
 
 private:
 
@@ -40,18 +48,6 @@ private:
 	};
 	/** Set while the recording and if the user did not disable this feature. */
 	TOptional<FRecordOverrides> RecordData;
-	
-	/** Sets the hitchless systems once all global state is set up (specifically UTakeRecorderSubsystem) */
-	void OnPostEngineInit();
-	/** Deinits the systems; */
-	void CleanUp();
-
-	void OnInitializeRecording(UTakeRecorder* TakeRecorder) { SetupEngineEnvironmentForRecording(); }
-	void OnFinishRecording(UTakeRecorder* TakeRecorder) { RecordData.Reset(); }
-	void OnCancelRecording(UTakeRecorder* TakeRecorder) { RecordData.Reset(); }
-
-	/** Sets up timecode and optionally the custom time step. */
-	void SetupEngineEnvironmentForRecording();
 };
 }
 
diff --git a/Engine/Plugins/VirtualProduction/Takes/Source/TakesCore/Private/TakeRecorderTimeProcessing.cpp b/Engine/Plugins/VirtualProduction/Takes/Source/TakesCore/Private/TakeRecorderTimeProcessing.cpp
--- a/Engine/Plugins/VirtualProduction/Takes/Source/TakesCore/Private/TakeRecorderTimeProcessing.cpp
+++ b/Engine/Plugins/VirtualProduction/Takes/Source/TakesCore/Private/TakeRecorderTimeProcessing.cpp
@@ -30,70 +30,14 @@ void ProcessRecordedTimes(ULevelSequence* InSequence, UMovieSceneTakeTrack* Take
 
 	if (FrameRange.IsSet())
 	{
-		TArray<int32> Hours, Minutes, Seconds, Frames;
-		TArray<FMovieSceneFloatValue> SubFrames;
-		TArray<FFrameNumber> Times;
-
-		Hours.Reserve(RecordedTimes.Num());
-		Minutes.Reserve(RecordedTimes.Num());
-		Seconds.Reserve(RecordedTimes.Num());
-		Frames.Reserve(RecordedTimes.Num());
-		SubFrames.Reserve(RecordedTimes.Num());
-		Times.Reserve(RecordedTimes.Num());
-
-		FFrameRate TickResolution = MovieScene->GetTickResolution();
-		FFrameRate DisplayRate = MovieScene->GetDisplayRate();
-
-		FFrameRate TCRate = TickResolution;
-		for (const TPair<FQualifiedFrameTime, FQualifiedFrameTime>& RecordedTimePair : RecordedTimes)
-		{
-			FFrameNumber FrameNumber = RecordedTimePair.Key.Time.FrameNumber;
-			if (!FrameRange.GetValue().Contains(FrameNumber))
-			{
-				continue;
-			}
-
-			FTimecode Timecode = RecordedTimePair.Value.ToTimecode();
-			TCRate = RecordedTimePair.Value.Rate;
-			Hours.Add(Timecode.Hours);
-			Minutes.Add(Timecode.Minutes);
-			Seconds.Add(Timecode.Seconds);
-			Frames.Add(Timecode.Frames);
-
-			FMovieSceneFloatValue SubFrame;
-			if (RecordedTimePair.Value.Time.GetSubFrame() > 0)
-			{
-				// If the Timecode provided gave us a subframe value then we should use that value.  Otherwise, we should compute
-				// the most appropriate value based on the timecode rate.
-				SubFrame.Value = RecordedTimePair.Value.Time.GetSubFrame();
-			}
-			else
-			{
-				FFrameTime FrameTime = FFrameRate::TransformTime(RecordedTimePair.Key.Time, TickResolution, DisplayRate);
-				FQualifiedFrameTime FrameTimeAsTimeCodeRate(FrameTime, TCRate);
-
-				SubFrame.Value = FrameTimeAsTimeCodeRate.Time.GetSubFrame();
-			}
-
-			SubFrame.InterpMode = ERichCurveInterpMode::RCIM_Linear;
-			SubFrames.Add(SubFrame);
-
-			Times.Add(FrameNumber);
-		}
-
-		Hours.Shrink();
-		Minutes.Shrink();
-		Seconds.Shrink();
-		Frames.Shrink();
-		SubFrames.Shrink();
-		Times.Shrink();
-
-		TakeSection->HoursCurve.Set(Times, Hours);
-		TakeSection->MinutesCurve.Set(Times, Minutes);
-		TakeSection->SecondsCurve.Set(Times, Seconds);
-		TakeSection->FramesCurve.Set(Times, Frames);
-		TakeSection->SubFramesCurve.Set(Times, SubFrames);
-		TakeSection->RateCurve.SetDefault(TCRate.AsDecimal());
+		FTimecodeExtractionData Data = ExtractIntoTimecodeArrays(MovieScene, *FrameRange, RecordedTimes);
+		TakeSection->HoursCurve.Set(Data.Times, MoveTemp(Data.Hours));
+		TakeSection->MinutesCurve.Set(Data.Times, MoveTemp(Data.Minutes));
+		TakeSection->SecondsCurve.Set(Data.Times, MoveTemp(Data.Seconds));
+		TakeSection->FramesCurve.Set(Data.Times, MoveTemp(Data.Frames));
+		// On the last use of Data.Times, we can move.
+		TakeSection->SubFramesCurve.Set(MoveTemp(Data.Times), MoveTemp(Data.SubFrames));
+		TakeSection->RateCurve.SetDefault(Data.TCRate.AsDecimal());
 	}
 
 	// Since the take section was created post recording here in this
@@ -129,4 +73,84 @@ void ProcessRecordedTimes(ULevelSequence* InSequence, UMovieSceneTakeTrack* Take
 		TakeSection->SetRange(TakeSection->GetAutoSizeRange().GetValue());
 	}
 }
+
+FTimecodeExtractionData ExtractIntoTimecodeArrays(
+	UMovieScene* InMovieScene, const TRange<FFrameNumber>& InRange, const FArrayOfRecordedTimePairs& InRecordedTimes,
+	ETimecodeExtractionFlags InFlags
+	)
+{
+	const int32 Num = InRecordedTimes.Num();
+	const auto Reserve = [InFlags, Num](auto& Array, ETimecodeExtractionFlags RequiredFlag)
+	{
+		if (EnumHasAnyFlags(InFlags, RequiredFlag))
+		{
+			Array.Reserve(Num);
+		}
+	};
+	const auto AddItem = [InFlags](auto& Array, const auto& Value, ETimecodeExtractionFlags RequiredFlag)
+	{
+		if (EnumHasAnyFlags(InFlags, RequiredFlag))
+		{
+			Array.Add(Value);
+		}
+	};
+	
+	TArray<int32> Hours, Minutes, Seconds, Frames;
+	TArray<FMovieSceneFloatValue> SubFrames;
+	TArray<FFrameNumber> Times;
+	Reserve(Hours, ETimecodeExtractionFlags::Hours);
+	Reserve(Minutes, ETimecodeExtractionFlags::Minutes);
+	Reserve(Seconds, ETimecodeExtractionFlags::Seconds);
+	Reserve(Frames, ETimecodeExtractionFlags::Frames);
+	Reserve(SubFrames, ETimecodeExtractionFlags::Subframes);
+	Reserve(Times, ETimecodeExtractionFlags::Times);
+	
+	const FFrameRate TickResolution = InMovieScene->GetTickResolution();
+	const FFrameRate DisplayRate = InMovieScene->GetDisplayRate();
+	FFrameRate TCRate = TickResolution;
+	for (const TPair<FQualifiedFrameTime, FQualifiedFrameTime>& RecordedTimePair : InRecordedTimes)
+	{
+		const FFrameNumber FrameNumber = RecordedTimePair.Key.Time.FrameNumber;
+		if (!InRange.Contains(FrameNumber))
+		{
+			continue;
+		}
+
+		const FTimecode Timecode = RecordedTimePair.Value.ToTimecode();
+		TCRate = RecordedTimePair.Value.Rate;
+		AddItem(Hours, Timecode.Hours, ETimecodeExtractionFlags::Hours);
+		AddItem(Minutes, Timecode.Minutes, ETimecodeExtractionFlags::Minutes);
+		AddItem(Seconds, Timecode.Seconds, ETimecodeExtractionFlags::Seconds);
+		AddItem(Frames, Timecode.Frames, ETimecodeExtractionFlags::Frames);
+
+		FMovieSceneFloatValue SubFrame;
+		if (RecordedTimePair.Value.Time.GetSubFrame() > 0)
+		{
+			// If the Timecode provided gave us a subframe value then we should use that value.  Otherwise, we should compute
+			// the most appropriate value based on the timecode rate.
+			SubFrame.Value = RecordedTimePair.Value.Time.GetSubFrame();
+		}
+		else
+		{
+			const FFrameTime FrameTime = FFrameRate::TransformTime(RecordedTimePair.Key.Time, TickResolution, DisplayRate);
+			const FQualifiedFrameTime FrameTimeAsTimeCodeRate(FrameTime, TCRate);
+			SubFrame.Value = FrameTimeAsTimeCodeRate.Time.GetSubFrame();
+		}
+
+		SubFrame.InterpMode = ERichCurveInterpMode::RCIM_Linear;
+		AddItem(SubFrames, SubFrame, ETimecodeExtractionFlags::Subframes);
+		AddItem(Times, FrameNumber, ETimecodeExtractionFlags::Times);
+	}
+
+	Hours.Shrink();
+	Minutes.Shrink();
+	Seconds.Shrink();
+	Frames.Shrink();
+	SubFrames.Shrink();
+	Times.Shrink();
+	return FTimecodeExtractionData
+	{
+		MoveTemp(Hours), MoveTemp(Minutes), MoveTemp(Seconds), MoveTemp(Frames), MoveTemp(SubFrames), MoveTemp(Times), TCRate
+	};
+}
 }
diff --git a/Engine/Plugins/VirtualProduction/Takes/Source/TakesCore/Public/TakeRecorderTimeProcessing.h b/Engine/Plugins/VirtualProduction/Takes/Source/TakesCore/Public/TakeRecorderTimeProcessing.h
--- a/Engine/Plugins/VirtualProduction/Takes/Source/TakesCore/Public/TakeRecorderTimeProcessing.h
+++ b/Engine/Plugins/VirtualProduction/Takes/Source/TakesCore/Public/TakeRecorderTimeProcessing.h
@@ -2,11 +2,13 @@
 
 #pragma once
 
+#include "Channels/MovieSceneFloatChannel.h"
 #include "Containers/Array.h"
 #include "Misc/QualifiedFrameTime.h"
 #include "Templates/Tuple.h"
 
 class ULevelSequence;
+class UMovieScene;
 class UMovieSceneTakeTrack;
 
 namespace UE::TakesCore
@@ -19,4 +21,38 @@ TAKESCORE_API void ProcessRecordedTimes(
 	ULevelSequence* InSequence, UMovieSceneTakeTrack* TakeTrack,
 	const TOptional<TRange<FFrameNumber>>& FrameRange, const FArrayOfRecordedTimePairs& RecordedTimes
 	);
+	
+enum class ETimecodeExtractionFlags : uint8
+{
+	None,
+
+	Hours = 1 << 0,
+	Minutes = 1 << 1,
+	Seconds = 1 << 2,
+	Frames = 1 << 3,
+	Subframes = 1 << 4,
+	
+	Times = 1 << 5,
+
+	All = Hours | Minutes | Seconds | Frames | Subframes | Times
+};
+ENUM_CLASS_FLAGS(ETimecodeExtractionFlags);
+
+// All arrays, except those excluded by the ETimecodeExtractionFlags flags, are equal length.
+struct FTimecodeExtractionData
+{
+	TArray<int32> Hours;
+	TArray<int32> Minutes;
+	TArray<int32> Seconds;
+	TArray<int32> Frames;
+	TArray<FMovieSceneFloatValue> SubFrames;
+	TArray<FFrameNumber> Times;
+	FFrameRate TCRate;
+};
+	
+/** Splits recorded InRecordedTimes up into arrays, one for each component. */
+[[nodiscard]] TAKESCORE_API FTimecodeExtractionData ExtractIntoTimecodeArrays(
+	UMovieScene* InMovieScene, const TRange<FFrameNumber>& InRange, const FArrayOfRecordedTimePairs& InRecordedTimes,
+	ETimecodeExtractionFlags InFlags = ETimecodeExtractionFlags::All
+	);
 }
\ No newline at end of file
