Take Recorder: Address BugHawk feedback.

diff --git a/Engine/Plugins/VirtualProduction/Takes/Source/TakeRecorder/Private/Timecode/Recording/FrameHitchDecorationRecorder.h b/Engine/Plugins/VirtualProduction/Takes/Source/TakeRecorder/Private/Timecode/Recording/FrameHitchDecorationRecorder.h
--- a/Engine/Plugins/VirtualProduction/Takes/Source/TakeRecorder/Private/Timecode/Recording/FrameHitchDecorationRecorder.h
+++ b/Engine/Plugins/VirtualProduction/Takes/Source/TakeRecorder/Private/Timecode/Recording/FrameHitchDecorationRecorder.h
@@ -14,9 +14,11 @@ class UTakeRecorder;
 namespace UE::TakeRecorder
 {
 /**
- * Creates 2 timecode tracks at the root Sequence:
- *  - UOriginalTimecodeMovieSceneTrack: Holds the unestimated timecode of the UTimecodeProvider underlying UTimecodeRegressionProvider.
- *  - URecordedTimecodeMovieSceneTrack: Holds the timecode of FApp::CurrentFrameTime.  */
+ * Records target and actual per-frame timecode and attached the data as UFrameHitchSceneDecoration to the root UMovieScene.
+ * Target timecode: Holds the unestimated timecode of the UTimecodeProvider underlying UTimecodeRegressionProvider.
+ * Actual timecode: Holds the timecode of FApp::CurrentFrameTime, which would usually (if the engine is set up) be what UTimecodeRegressionProvider
+ * estimates.
+ */
 class FFrameHitchDecorationRecorder : public FNoncopyable
 {
 public:
