Take Recorder: Fix OnStopped event not firing in subsystem.

diff --git a/Engine/Plugins/VirtualProduction/Takes/Source/TakeRecorder/Private/Recorder/TakeRecorderSubsystemImplementation.cpp b/Engine/Plugins/VirtualProduction/Takes/Source/TakeRecorder/Private/Recorder/TakeRecorderSubsystemImplementation.cpp
--- a/Engine/Plugins/VirtualProduction/Takes/Source/TakeRecorder/Private/Recorder/TakeRecorderSubsystemImplementation.cpp
+++ b/Engine/Plugins/VirtualProduction/Takes/Source/TakeRecorder/Private/Recorder/TakeRecorderSubsystemImplementation.cpp
@@ -143,16 +143,9 @@ void UTakeRecorderSubsystemImplementation::SetTargetSequence(const FTakeRecorder
 	TakeRecorderModule.GetLastLevelSequenceProvider().BindUObject(this, &UTakeRecorderSubsystemImplementation::SetLastLevelSequence);
 	
 	// If a recording is currently underway, initialize to that now
-	if (UTakeRecorder* ActiveRecorder = UTakeRecorder::GetActiveRecorder())
+	if (const UTakeRecorder* ActiveRecorder = UTakeRecorder::GetActiveRecorder())
 	{
 		RecordingLevelSequence = ActiveRecorder->GetSequence();
-
-		ActiveRecorder->OnRecordingStopped().Remove(OnRecordingStoppedHandle);
-		ActiveRecorder->OnRecordingFinished().Remove(OnRecordingFinishedHandle);
-		ActiveRecorder->OnRecordingCancelled().Remove(OnRecordingCancelledHandle);
-		OnRecordingStoppedHandle = ActiveRecorder->OnRecordingStopped().AddUObject(this, &UTakeRecorderSubsystemImplementation::OnRecordingStopped);
-		OnRecordingFinishedHandle = ActiveRecorder->OnRecordingFinished().AddUObject(this, &UTakeRecorderSubsystemImplementation::OnRecordingFinished);
-		OnRecordingCancelledHandle = ActiveRecorder->OnRecordingCancelled().AddUObject(this, &UTakeRecorderSubsystemImplementation::OnRecordingCancelled);
 	}
 	else
 	{
@@ -970,8 +963,9 @@ void UTakeRecorderSubsystemImplementation::OnRecordingInitialized(UTakeRecorder*
 	// Recache the meta-data here since we know that the sequence has probably changed as a result of the recording being started
 	CacheMetaData();
 	OnRecordingFinishedHandle = Recorder->OnRecordingFinished().AddUObject(this, &UTakeRecorderSubsystemImplementation::OnRecordingFinished);
+	OnRecordingStoppedHandle = Recorder->OnRecordingStopped().AddUObject(this, &UTakeRecorderSubsystemImplementation::OnRecordingStopped);
 	OnRecordingCancelledHandle = Recorder->OnRecordingCancelled().AddUObject(this, &UTakeRecorderSubsystemImplementation::OnRecordingCancelled);
-
+	
 	WITH_OWNING_SUBSYSTEM(GetOnRecordingInitializedEvent().Broadcast(Recorder));
 	WITH_OWNING_SUBSYSTEM(TakeRecorderInitialized.Broadcast());
 }
