Take Recorder: Add a cvar TakeRecorder.AllowPossessablePIEObjects to bypass disallowing recording PIE objects as possessables. The default is to not allow this (no change in behavior)

This can be useful when recording the player which maybe cause gleo

diff --git a/Engine/Plugins/VirtualProduction/Takes/Source/TakeRecorderSources/Private/TakeRecorderActorSource.cpp b/Engine/Plugins/VirtualProduction/Takes/Source/TakeRecorderSources/Private/TakeRecorderActorSource.cpp
--- a/Engine/Plugins/VirtualProduction/Takes/Source/TakeRecorderSources/Private/TakeRecorderActorSource.cpp
+++ b/Engine/Plugins/VirtualProduction/Takes/Source/TakeRecorderSources/Private/TakeRecorderActorSource.cpp
@@ -62,6 +62,12 @@ static const FName MovieSceneSectionRecorderFactoryName("MovieSceneTrackRecorder
 static const FName SequencerTrackClassMetadataName("SequencerTrackClass");
 static const FName DoNotRecordTag("DoNotRecord");
 
+static TAutoConsoleVariable<bool> CVarTakeRecorderAllowPossessablePIEObjects(
+	TEXT("TakeRecorder.AllowPossessablePIEObjects"),
+	0,
+	TEXT("0: Don't allow recording PIE objects as possessable\n1: Override to allow recording PIE Objects as possessable (note the bindings will be broken!)"),
+	ECVF_Default);
+
 UTakeRecorderSource* UTakeRecorderActorSource::AddSourceForActor(AActor* InActor, UTakeRecorderSources* InSources)
 {
 	if (InSources == nullptr)
@@ -252,7 +258,8 @@ TArray<UTakeRecorderSource*> UTakeRecorderActorSource::PreRecording(ULevelSequen
 	if (GetRecordToPossessable())
 	{
 		// If a user adds a PIE-only instance as a recordable object, they can't record this to a possessable (because the binding will be broken when they exit PIE).
-		if (Target->GetWorld()->WorldType != EWorldType::Editor && GEditor && !GEditor->ObjectsThatExistInEditorWorld.Get(Target.Get()))
+		if (!CVarTakeRecorderAllowPossessablePIEObjects->GetBool() &&
+			Target->GetWorld()->WorldType != EWorldType::Editor && GEditor && !GEditor->ObjectsThatExistInEditorWorld.Get(Target.Get()))
 		{
 			UE_LOG(LogTakesCore, Warning, TEXT("Attempted to record an actor that does not exist in the editor world to a possessable. Forcing recording of %s as a Spawnable so that the resulting binding is not broken!"), *Target->GetName());
 			RecordType = ETakeRecorderActorRecordType::Spawnable;
