Take Recorder: For hitch analysis, fix frame's target timecode sometimes being sampled with a value diferent from the value sampled at the beginning of frame.

The problem is that some timecode providers's GetQualifiedFrameTime, specifically USystemTimeTimecodeProvider, always returns the latest timecode. So if enough time had passed, at the beginning of the frame the timecode could be 13:55:43:01 but later in the same frame 13:55:43:02. That would make hitch analysis show timecode skips. Now, the target frame rate is set to whatever it was at the beginning of the frame, which is the value we sampled in the UTimecodeRegressionProvider at the beginning of the frame.

diff --git a/Engine/Plugins/VirtualProduction/Takes/Source/TakeRecorder/Private/Timecode/HitchlessProtectionRootLogic.cpp b/Engine/Plugins/VirtualProduction/Takes/Source/TakeRecorder/Private/Timecode/HitchlessProtectionRootLogic.cpp
--- a/Engine/Plugins/VirtualProduction/Takes/Source/TakeRecorder/Private/Timecode/HitchlessProtectionRootLogic.cpp
+++ b/Engine/Plugins/VirtualProduction/Takes/Source/TakeRecorder/Private/Timecode/HitchlessProtectionRootLogic.cpp
@@ -68,15 +68,15 @@ void FHitchlessProtectionRootLogic::OnInitializeRecording(UTakeRecorder* TakeRec
 		return;
 	}
 
-	const TOptional<TNonNullPtr<UTimecodeProvider>> UnderlyingTimecodeProvider = TimecodeRegressionRecording.SetupEngineEnvironmentForRecording(
+	const FTimecodeEstimationSetupResult TimecodeSetupResult = TimecodeRegressionRecording.SetupEngineEnvironmentForRecording(
 		Private::CVarTimecodeLinearRegressionSeconds.GetValueOnAnyThread()
 		);
-	if (!UnderlyingTimecodeProvider)
+	if (!TimecodeSetupResult)
 	{
 		return;
 	}
 	
-	RecordingState.Emplace(*UnderlyingTimecodeProvider, TakeRecorder);
+	RecordingState.Emplace(TimecodeSetupResult.Estimator, TakeRecorder);
 }
 
 void FHitchlessProtectionRootLogic::OnFinishRecording(UTakeRecorder* TakeRecorder)
diff --git a/Engine/Plugins/VirtualProduction/Takes/Source/TakeRecorder/Private/Timecode/HitchlessProtectionRootLogic.h b/Engine/Plugins/VirtualProduction/Takes/Source/TakeRecorder/Private/Timecode/HitchlessProtectionRootLogic.h
--- a/Engine/Plugins/VirtualProduction/Takes/Source/TakeRecorder/Private/Timecode/HitchlessProtectionRootLogic.h
+++ b/Engine/Plugins/VirtualProduction/Takes/Source/TakeRecorder/Private/Timecode/HitchlessProtectionRootLogic.h
@@ -37,8 +37,8 @@ private:
 		/** Records the underlying timecode provider and the one that is estimated by UTimecodeRegressionProvider. */
 		FFrameHitchDecorationRecorder TimecodeRecording;
 
-		explicit FRecordingState(TNotNull<UTimecodeProvider*> InRecordedProvider, TNotNull<UTakeRecorder*> InTakeRecorder)
-			: TimecodeRecording(InRecordedProvider, InTakeRecorder)
+		explicit FRecordingState(TNotNull<UTimecodeRegressionProvider*> InEstimator, TNotNull<UTakeRecorder*> InTakeRecorder)
+			: TimecodeRecording(InEstimator, InTakeRecorder)
 		{}
 	};
 	/** Set while recording. */
diff --git a/Engine/Plugins/VirtualProduction/Takes/Source/TakeRecorder/Private/Timecode/Recording/FrameHitchDecorationRecorder.cpp b/Engine/Plugins/VirtualProduction/Takes/Source/TakeRecorder/Private/Timecode/Recording/FrameHitchDecorationRecorder.cpp
--- a/Engine/Plugins/VirtualProduction/Takes/Source/TakeRecorder/Private/Timecode/Recording/FrameHitchDecorationRecorder.cpp
+++ b/Engine/Plugins/VirtualProduction/Takes/Source/TakeRecorder/Private/Timecode/Recording/FrameHitchDecorationRecorder.cpp
@@ -3,11 +3,13 @@
 #include "FrameHitchDecorationRecorder.h"
 
 #include "Engine/TimecodeProvider.h"
+#include "Estimation/TimecodeRegressionProvider.h"
 #include "Hitching/FrameHitchSceneDecoration.h"
 #include "LevelSequence.h"
 #include "MovieScene.h"
 #include "Recorder/TakeRecorder.h"
 #include "TakeRecorderSources.h"
+#include "TakesCoreLog.h"
 
 #define LOCTEXT_NAMESPACE "FTimecodeTrackRecording"
 
@@ -51,8 +53,10 @@ static TOptional<TRange<FFrameNumber>> GetGlobalFrameRange(const UMovieScene* Mo
 	return FrameRange;
 }
 }
-FFrameHitchDecorationRecorder::FFrameHitchDecorationRecorder(TNotNull<UTimecodeProvider*> InRecordedProvider, TNotNull<UTakeRecorder*> InTakeRecorder)
-	: WeakTimecodeProvider(InRecordedProvider)
+FFrameHitchDecorationRecorder::FFrameHitchDecorationRecorder(
+	TNotNull<UTimecodeRegressionProvider*> InTimecodeEstimator, TNotNull<UTakeRecorder*> InTakeRecorder
+	)
+	: WeakTimecodeEstimator(InTimecodeEstimator)
 	, WeakTakeRecorder(InTakeRecorder)
 {
 	InTakeRecorder->OnTickRecording().AddRaw(this, &FFrameHitchDecorationRecorder::OnTickRecording);
@@ -67,14 +71,28 @@ FFrameHitchDecorationRecorder::~FFrameHitchDecorationRecorder()
 		TakeRecorder->OnRecordingStopped().RemoveAll(this);
 	}
 }
-
+	
 void FFrameHitchDecorationRecorder::OnTickRecording(UTakeRecorder* InTakeRecorder, const FQualifiedFrameTime& InCurrentFrameTime)
 {
-	if (UTimecodeProvider* TimecodeProvider = WeakTimecodeProvider.Get())
-	{
-		const FQualifiedFrameTime RealTimeCode = TimecodeProvider->GetQualifiedFrameTime();
-		TargetTimecodeValues.Emplace(InCurrentFrameTime, RealTimeCode);
-	}
+	UTimecodeRegressionProvider* TimecodeProvider = WeakTimecodeEstimator.Get();
+	
+	// Intention: We need to ask the underlying timecode provider what timecode this frame is supposed to have.
+	// We CANNOT just call the underlying UTimecodeProvider::GetQualifiedTime because some providers, specifically USystemTimeTimecodeProvider, will resample
+	// the timecode and may return a newer timecode. That's problematic as some time may have elapsed since the beginning of the frame and now (so we'd get a different timecode than we got earlier in this frame).
+	// So we ask the UTimecodeRegressionProvider what it sampled to make sure we use the timecode that this frame was supposed to have at the beginning of this frame.
+	// For the target timecode, we shoul make sure to use exatly the timecode that was sampled at the beginning of this frame... when FApp::CurrentFrameTime was set.
+	const TOptional<FQualifiedFrameTime> LastSampled = TimecodeProvider ? TimecodeProvider->GetLastSampledFrameTime() : TOptional<FQualifiedFrameTime>{};
+	UE_CLOG(!LastSampled, LogTakesCore, Warning,
+		TEXT("No timecode was sampled by the linear regression. Hitch analysis data for frame %s will use actual frame time instead of underlying target frame time. Hitch visualization will be inaccurate for this frame."),
+		*InCurrentFrameTime.ToTimecode().ToString()
+		);
+	
+	TargetTimecodeValues.Emplace(
+		InCurrentFrameTime,
+		// It's important for analysis, we always add a timecode so TargetTimecodeValues has equal length as UTakeRecorderSources::RecordedTimes.
+		// If for whatever reason we cannot determine the "correct" target timecode, we'll just pretend the recorded one is the target.
+		LastSampled.Get(InCurrentFrameTime)
+		);
 }
 
 void FFrameHitchDecorationRecorder::OnRecordingStopped(UTakeRecorder* TakeRecorder) const
diff --git a/Engine/Plugins/VirtualProduction/Takes/Source/TakeRecorder/Private/Timecode/Recording/FrameHitchDecorationRecorder.h b/Engine/Plugins/VirtualProduction/Takes/Source/TakeRecorder/Private/Timecode/Recording/FrameHitchDecorationRecorder.h
--- a/Engine/Plugins/VirtualProduction/Takes/Source/TakeRecorder/Private/Timecode/Recording/FrameHitchDecorationRecorder.h
+++ b/Engine/Plugins/VirtualProduction/Takes/Source/TakeRecorder/Private/Timecode/Recording/FrameHitchDecorationRecorder.h
@@ -8,7 +8,7 @@
 #include "Templates/UnrealTemplate.h"
 #include "UObject/WeakObjectPtrTemplates.h"
 
-class UTimecodeProvider;
+class UTimecodeRegressionProvider;
 class UTakeRecorder;
 
 namespace UE::TakeRecorder
@@ -23,18 +23,23 @@ class FFrameHitchDecorationRecorder : public FNoncopyable
 {
 public:
 	
-	explicit FFrameHitchDecorationRecorder(TNotNull<UTimecodeProvider*> InRecordedProvider, TNotNull<UTakeRecorder*> InTakeRecorder);
+	explicit FFrameHitchDecorationRecorder(TNotNull<UTimecodeRegressionProvider*> InTimecodeEstimator, TNotNull<UTakeRecorder*> InTakeRecorder);
 	~FFrameHitchDecorationRecorder();
 
 private:
 
-	/** The timecode provider being recorded. */
-	const TWeakObjectPtr<UTimecodeProvider> WeakTimecodeProvider;
+	/** The timecode provider that's doing the estimation of frame times. */
+	const TWeakObjectPtr<UTimecodeRegressionProvider> WeakTimecodeEstimator;
 
 	/** Take recorder for which we are recording. */
 	const TWeakObjectPtr<UTakeRecorder> WeakTakeRecorder;
 	
-	/** Passed in to TakeRecorderSourceHelpers::ProcessRecordedTimes to generate the track for original timecode. */
+	/**
+	 * Passed in to TakeRecorderSourceHelpers::ProcessRecordedTimes to generate the track for original timecode.
+	 * 
+	 * We make sure this has equal length as UTakeRecorderSources::RecordedTimes. We achieve this by sampling TargetTimecodeValues whenever
+	 * UTakeRecorderSources::RecordedTimes is sampled which is when UTakeRecorder::OnTickRecording is invoked.
+	 */
 	TakesCore::FArrayOfRecordedTimePairs TargetTimecodeValues;
 
 	/** Called when InTakeRecorder ticks during recording. */
diff --git a/Engine/Plugins/VirtualProduction/Takes/Source/TakeRecorder/Private/Timecode/Regression/TimecodeRegressionRecordSetup.cpp b/Engine/Plugins/VirtualProduction/Takes/Source/TakeRecorder/Private/Timecode/Regression/TimecodeRegressionRecordSetup.cpp
--- a/Engine/Plugins/VirtualProduction/Takes/Source/TakeRecorder/Private/Timecode/Regression/TimecodeRegressionRecordSetup.cpp
+++ b/Engine/Plugins/VirtualProduction/Takes/Source/TakeRecorder/Private/Timecode/Regression/TimecodeRegressionRecordSetup.cpp
@@ -33,32 +33,36 @@ static UTimecodeRegressionProvider* CreateTimecodeEstimator(UTimecodeProvider* I
 }
 }
 
-TOptional<TNonNullPtr<UTimecodeProvider>> FTimecodeRegressionRecordSetup::SetupEngineEnvironmentForRecording(float InRegressionFramesInSeconds)
+FTimecodeEstimationSetupResult FTimecodeRegressionRecordSetup::SetupEngineEnvironmentForRecording(float InRegressionFramesInSeconds)
 {
+	FTimecodeEstimationSetupResult Result;
+	
 	UTimecodeProvider* CurrentTimecodeProvider = GEngine ? GEngine->GetTimecodeProvider() : nullptr;
 	if (!CurrentTimecodeProvider)
 	{
-		return {};
+		return Result;
 	}
 
 	UEngineCustomTimeStep* PreexistingTimestep = GEngine->GetCustomTimeStep();
 	if (PreexistingTimestep && !UTimecodeRegressionProvider::SupportsTimestep(PreexistingTimestep))
 	{
 		ShowIncompatibleTimestepNotification(PreexistingTimestep);
-		return {};
+		return Result;
 	}
 
+	UTimecodeRegressionProvider* Estimator = Private::CreateTimecodeEstimator(CurrentTimecodeProvider, InRegressionFramesInSeconds);
 	if (PreexistingTimestep)
 	{
 		// Use pre-configured custom timestep
-		RecordData.Emplace(Private::CreateTimecodeEstimator(CurrentTimecodeProvider, InRegressionFramesInSeconds));
+		RecordData.Emplace(Estimator);
 	}
 	else if (UEngineCustomTimeStep* DefaultTimestep = AskUserToCreateFallbackTimestep())
 	{
 		// There was no pre-configured custom timestep... clean it up when recording finishes.
-		RecordData.Emplace(DefaultTimestep, Private::CreateTimecodeEstimator(CurrentTimecodeProvider, InRegressionFramesInSeconds));
+		RecordData.Emplace(DefaultTimestep, Estimator);
 	}
 
-	return TNonNullPtr<UTimecodeProvider>(CurrentTimecodeProvider);
+	Result.Estimator = Estimator;
+	return Result;
 }
 }
\ No newline at end of file
diff --git a/Engine/Plugins/VirtualProduction/Takes/Source/TakeRecorder/Private/Timecode/Regression/TimecodeRegressionRecordSetup.h b/Engine/Plugins/VirtualProduction/Takes/Source/TakeRecorder/Private/Timecode/Regression/TimecodeRegressionRecordSetup.h
--- a/Engine/Plugins/VirtualProduction/Takes/Source/TakeRecorder/Private/Timecode/Regression/TimecodeRegressionRecordSetup.h
+++ b/Engine/Plugins/VirtualProduction/Takes/Source/TakeRecorder/Private/Timecode/Regression/TimecodeRegressionRecordSetup.h
@@ -11,6 +11,16 @@
 
 namespace UE::TakeRecorder
 {
+/** Holds info about when the engine environment has been set up to estimate timecode. */
+struct FTimecodeEstimationSetupResult
+{
+	/** The timecode estimator that was set up. */
+	UTimecodeRegressionProvider* Estimator = nullptr;
+
+	bool IsSuccess() const { return Estimator != nullptr; }
+	operator bool() const { return IsSuccess(); }
+};
+	
 /**
  * Sets up the UEngine::TimecodeProvider so timecode is not skipped when an engine frame takes longer than expected ("hitches").
  * This is root class is responsible for setting up this for this workflow and acts as a mediator for setting up relevant systems.
@@ -25,7 +35,7 @@ public:
 	 *	Example: 1s worth of frames means a buffer size of 24 timecode samples if the UTakeRecorder is set to record at 24 FPS.
 	 * @return The real timecode provider that we will be estimating. Unset if initialization failed.
 	 */
-	TOptional<TNonNullPtr<UTimecodeProvider>> SetupEngineEnvironmentForRecording(float InRegressionFramesInSeconds);
+	FTimecodeEstimationSetupResult SetupEngineEnvironmentForRecording(float InRegressionFramesInSeconds);
 
 	/** Clears the time step and timecode provider overrides. */
 	void CleanupRecording() { RecordData.Reset(); }
diff --git a/Engine/Source/Runtime/TimeManagement/Private/Estimation/TimecodeEstimator.cpp b/Engine/Source/Runtime/TimeManagement/Private/Estimation/TimecodeEstimator.cpp
--- a/Engine/Source/Runtime/TimeManagement/Private/Estimation/TimecodeEstimator.cpp
+++ b/Engine/Source/Runtime/TimeManagement/Private/Estimation/TimecodeEstimator.cpp
@@ -96,12 +96,12 @@ FTimecodeEstimator::FTimecodeEstimator(
 	ensure(InNumSamples > 0); 
 }
 
-void FTimecodeEstimator::FetchAndUpdate()
+TOptional<FFetchAndUpdateStats> FTimecodeEstimator::FetchAndUpdate()
 {
 	const TOptional<double> ClockTime = EngineCustomTimeStep.GetUnderlyingClockTime_AnyThread();
 	if (!ClockTime)
 	{
-		return;
+		return {};
 	}
 
 	if (!StartClockTime)
@@ -135,6 +135,8 @@ void FTimecodeEstimator::FetchAndUpdate()
 		RelativeTime, *FQualifiedFrameTime(FrameTime, LastFrameRate).ToTimecode().ToString(bForceSignDisplay, bDisplaySubframe),
 		*ClockTime, FApp::GetCurrentTime()
 		);
+
+	return FFetchAndUpdateStats{ CurrentFrameTime };
 }
 
 FQualifiedFrameTime FTimecodeEstimator::EstimateFrameTime() const
diff --git a/Engine/Source/Runtime/TimeManagement/Private/Estimation/TimecodeRegressionProvider.cpp b/Engine/Source/Runtime/TimeManagement/Private/Estimation/TimecodeRegressionProvider.cpp
--- a/Engine/Source/Runtime/TimeManagement/Private/Estimation/TimecodeRegressionProvider.cpp
+++ b/Engine/Source/Runtime/TimeManagement/Private/Estimation/TimecodeRegressionProvider.cpp
@@ -21,11 +21,18 @@ void UTimecodeRegressionProvider::SetTimecodeImplementation(UTimecodeProvider* I
 	}
 }
 
+TOptional<FQualifiedFrameTime> UTimecodeRegressionProvider::GetLastSampledFrameTime() const
+{
+	return EstimationState && EstimationState->LastUpdateStats
+		? EstimationState->LastUpdateStats->UnderlyingFrameTime
+		: TOptional<FQualifiedFrameTime>();
+}
+
 void UTimecodeRegressionProvider::FetchAndUpdate()
 {
-	if (TimecodeEstimator)
+	if (EstimationState)
 	{
-		TimecodeEstimator->FetchAndUpdate();
+		EstimationState->LastUpdateStats = EstimationState->TimecodeEstimator.FetchAndUpdate();
 	}
 	// This is just a sensible fallback if ReinitRegressionBuffer failed due to unsupported timestep class
 	else if (TimecodeImpl)
@@ -36,8 +43,8 @@ void UTimecodeRegressionProvider::FetchAndUpdate()
 
 FQualifiedFrameTime UTimecodeRegressionProvider::GetQualifiedFrameTime() const
 {
-	const FQualifiedFrameTime Time = TimecodeEstimator
-		? TimecodeEstimator->EstimateFrameTime()
+	const FQualifiedFrameTime Time = EstimationState
+		? EstimationState->TimecodeEstimator.EstimateFrameTime()
 		// This is just a sensible fallback if ReinitRegressionBuffer failed due to unsupported timestep class or missing call to SetTimecodeImplementation. 
 		: TimecodeImpl ? TimecodeImpl->GetQualifiedFrameTime() : FQualifiedFrameTime();
 	return Time;
@@ -93,10 +100,10 @@ void UTimecodeRegressionProvider::ReinitRegressionBuffer()
 	
 	if (ClockedTimeStep)
 	{
-		TimecodeEstimator.Emplace(FMath::Clamp(NumSampledFrames, 1, 65536), *TimecodeImpl, *ClockedTimeStep);
+		EstimationState.Emplace(FMath::Clamp(NumSampledFrames, 1, 65536), *TimecodeImpl, *ClockedTimeStep);
 	}
 	else
 	{
-		TimecodeEstimator.Reset();
+		EstimationState.Reset();
 	}
 }
diff --git a/Engine/Source/Runtime/TimeManagement/Public/Estimation/TimecodeEstimator.h b/Engine/Source/Runtime/TimeManagement/Public/Estimation/TimecodeEstimator.h
--- a/Engine/Source/Runtime/TimeManagement/Public/Estimation/TimecodeEstimator.h
+++ b/Engine/Source/Runtime/TimeManagement/Public/Estimation/TimecodeEstimator.h
@@ -15,6 +15,13 @@ struct FQualifiedFrameTime;
 
 namespace UE::TimeManagement::TimecodeEstimation
 {
+/** Misc information about how the timecode estimation was updated. */
+struct FFetchAndUpdateStats
+{
+	/** The frame time that was sampled from the underlying timecode provider. */
+	FQualifiedFrameTime UnderlyingFrameTime;
+};
+	
 /**
  * Estimates the current timecode based on an IClockedTimeStep implementation, which is designed to be an UEngineCustomTimeStep.
  *
@@ -68,8 +75,13 @@ public:
 		UTimecodeProvider& InTimecode UE_LIFETIMEBOUND, IClockedTimeStep& InEngineCustomTimeStep UE_LIFETIMEBOUND
 		);
 
-	/** Samples the current timecode and associates it with the underlying clock value. */
-	UE_API void FetchAndUpdate();
+	/**
+	 * Samples the current timecode and associates it with the underlying clock value.
+	 * 
+	 * @return Metadata about how update has occured, e.g. the "real", frame time sampled from the timecode provider.
+	 *	Unset if the custom time step's clock could not be read.
+	 */
+	UE_API TOptional<FFetchAndUpdateStats> FetchAndUpdate();
 	
 	/** Estimates what the current frame time should be given FApp::CurrentTime's value. */
 	UE_API FQualifiedFrameTime EstimateFrameTime() const;
diff --git a/Engine/Source/Runtime/TimeManagement/Public/Estimation/TimecodeRegressionProvider.h b/Engine/Source/Runtime/TimeManagement/Public/Estimation/TimecodeRegressionProvider.h
--- a/Engine/Source/Runtime/TimeManagement/Public/Estimation/TimecodeRegressionProvider.h
+++ b/Engine/Source/Runtime/TimeManagement/Public/Estimation/TimecodeRegressionProvider.h
@@ -41,6 +41,15 @@ public:
 	/** Sets the actual time code provider. You must call this before registering with the engine, i.e. before Initialize() is called. */
 	UE_API void SetTimecodeImplementation(UTimecodeProvider* InTimecodeImpl);
 
+	/** @return Gets the frame time that was last appended to the linear regression. */
+	UE_API TOptional<FQualifiedFrameTime> GetLastSampledFrameTime() const;
+
+	/** @return The timecode provider which we sample and use to estimate timecode. */
+	UTimecodeProvider* GetUnderlyingTimecodeProvider() const
+	{
+		return TimecodeImpl;
+	}
+
 	//~ Begin UTimecodeProvider Interface
 	virtual bool FetchTimecode(FQualifiedFrameTime& OutFrameTime) override
 	{
@@ -74,8 +83,23 @@ private:
 	UPROPERTY()
 	TObjectPtr<UTimecodeProvider> TimecodeImpl;
 
-	/** Estimates the timecode based on FApp::CurrentTime. */
-	TOptional<UE::TimeManagement::TimecodeEstimation::FTimecodeEstimator> TimecodeEstimator;
+	/** Data valid for the lifetime of the FTimecodeEstimator. */
+	struct FTimecodeEstimationState
+	{
+		/** Estimates the timecode based on FApp::CurrentTime. */
+		UE::TimeManagement::TimecodeEstimation::FTimecodeEstimator TimecodeEstimator;
+
+		/** Result from last TimecodeEstimator.FetchAndUpdate call. */
+		TOptional<UE::TimeManagement::TimecodeEstimation::FFetchAndUpdateStats> LastUpdateStats;
+
+		explicit FTimecodeEstimationState(
+			SIZE_T InNumSamples, UTimecodeProvider& InTimecode UE_LIFETIMEBOUND, IClockedTimeStep& InEngineCustomTimeStep UE_LIFETIMEBOUND
+			)
+			: TimecodeEstimator(InNumSamples, InTimecode, InEngineCustomTimeStep)
+		{}
+	};
+	/** Data set while we're estimating timecode. */
+	TOptional<FTimecodeEstimationState> EstimationState;
 
 	/** @return Whether Initialize has been called. */
 	bool IsInitialized() const
