Take Recorder: Hitch Visualization using Sequencer's FViewModel. Uses recorded timecode from UFrameHitchSceneDecoration to draw in Sequencer:
- A red marker on frames which skipped a timecode
- A yellow range on frames that are behind the target timecode

diff --git a/Engine/Plugins/VirtualProduction/Takes/Source/TakeRecorder/Private/TakeRecorderCommands.cpp b/Engine/Plugins/VirtualProduction/Takes/Source/TakeRecorder/Private/TakeRecorderCommands.cpp
--- a/Engine/Plugins/VirtualProduction/Takes/Source/TakeRecorder/Private/TakeRecorderCommands.cpp
+++ b/Engine/Plugins/VirtualProduction/Takes/Source/TakeRecorder/Private/TakeRecorderCommands.cpp
@@ -14,6 +14,7 @@ void FTakeRecorderCommands::RegisterCommands()
 {
 	UI_COMMAND(StartRecording, "StartRecording", "Start recording", EUserInterfaceActionType::Button, FInputChord(EModifierKey::Control | EModifierKey::Shift, EKeys::R));
 	UI_COMMAND(StopRecording, "StopRecording", "Stop recording", EUserInterfaceActionType::Button, FInputChord(EModifierKey::Control | EModifierKey::Shift, EKeys::S));
+	UI_COMMAND(ClearRecordingIntegrityData, "Clear Recording Integrity Markers and Ranges", "Clear Recording Integrity Markers and Ranges", EUserInterfaceActionType::Button, FInputChord());
 }
 
 #undef LOCTEXT_NAMESPACE
diff --git a/Engine/Plugins/VirtualProduction/Takes/Source/TakeRecorder/Private/TakeRecorderCommands.h b/Engine/Plugins/VirtualProduction/Takes/Source/TakeRecorder/Private/TakeRecorderCommands.h
--- a/Engine/Plugins/VirtualProduction/Takes/Source/TakeRecorder/Private/TakeRecorderCommands.h
+++ b/Engine/Plugins/VirtualProduction/Takes/Source/TakeRecorder/Private/TakeRecorderCommands.h
@@ -13,6 +13,8 @@ public:
 	TSharedPtr<FUICommandInfo> StartRecording;
 	TSharedPtr<FUICommandInfo> StopRecording;
 
+	TSharedPtr<FUICommandInfo> ClearRecordingIntegrityData;
+
 	/** Initialize commands */
 	virtual void RegisterCommands() override;
 };
diff --git a/Engine/Plugins/VirtualProduction/Takes/Source/TakeRecorder/Private/TakeRecorderModule.cpp b/Engine/Plugins/VirtualProduction/Takes/Source/TakeRecorder/Private/TakeRecorderModule.cpp
--- a/Engine/Plugins/VirtualProduction/Takes/Source/TakeRecorder/Private/TakeRecorderModule.cpp
+++ b/Engine/Plugins/VirtualProduction/Takes/Source/TakeRecorder/Private/TakeRecorderModule.cpp
@@ -335,6 +335,12 @@ void FTakeRecorderModule::StartupModule()
 	RegisterSettings();
 	RegisterSerializedRecorder();
 	TimecodeManagement = MakeUnique<UE::TakeRecorder::FHitchlessProtectionRootLogic>();
+#if WITH_EDITOR
+	if (GIsEditor)
+	{
+		SequencerHitchVisualizer = MakeUnique<UE::TakeRecorder::FSequencerHitchVisualizer>();
+	}
+#endif
 
 #if WITH_EDITOR
 	if (GIsEditor)
@@ -375,6 +381,9 @@ void FTakeRecorderModule::ShutdownModule()
 	UnregisterSettings();
 	UnregisterSerializedRecorder();
 	TimecodeManagement.Reset();
+#if WITH_EDITOR
+	SequencerHitchVisualizer.Reset();
+#endif
 }
 
 UTakePreset* FTakeRecorderModule::GetPendingTake() const
diff --git a/Engine/Plugins/VirtualProduction/Takes/Source/TakeRecorder/Private/TakeRecorderModule.h b/Engine/Plugins/VirtualProduction/Takes/Source/TakeRecorder/Private/TakeRecorderModule.h
--- a/Engine/Plugins/VirtualProduction/Takes/Source/TakeRecorder/Private/TakeRecorderModule.h
+++ b/Engine/Plugins/VirtualProduction/Takes/Source/TakeRecorder/Private/TakeRecorderModule.h
@@ -2,10 +2,12 @@
 
 #pragma once
 
-#include "Templates/SharedPointer.h"
 #include "Delegates/IDelegateInstance.h"
-#include "UObject/GCObject.h"
 #include "ITakeRecorderModule.h"
+#include "Templates/SharedPointer.h"
+#include "Timecode/HitchlessProtectionRootLogic.h"
+#include "Timecode/Visualization/Sequencer/SequencerHitchVisualizer.h"
+#include "UObject/GCObject.h"
 #include "Timecode/HitchlessProtectionRootLogic.h"
 #include "UObject/WeakObjectPtrTemplates.h"
 
@@ -100,6 +102,9 @@ private:
 	 * gracefully.
 	 */
 	TUniquePtr<UE::TakeRecorder::FHitchlessProtectionRootLogic> TimecodeManagement;
+#if WITH_EDITOR
+	TUniquePtr<UE::TakeRecorder::FSequencerHitchVisualizer> SequencerHitchVisualizer;
+#endif
 
 	/** Extra class default objects that should appear in the Take Recorder panel. */
 	TArray<TWeakObjectPtr<>> ExternalObjects;
diff --git a/Engine/Plugins/VirtualProduction/Takes/Source/TakeRecorder/Private/TakeRecorderStyle.cpp b/Engine/Plugins/VirtualProduction/Takes/Source/TakeRecorder/Private/TakeRecorderStyle.cpp
--- a/Engine/Plugins/VirtualProduction/Takes/Source/TakeRecorder/Private/TakeRecorderStyle.cpp
+++ b/Engine/Plugins/VirtualProduction/Takes/Source/TakeRecorder/Private/TakeRecorderStyle.cpp
@@ -6,6 +6,7 @@
 #include "Styling/SlateTypes.h"
 #include "Styling/CoreStyle.h"
 #include "Styling/AppStyle.h"
+#include "Styling/SlateStyleMacros.h"
 #include "Styling/StyleColors.h"
 
 FName FTakeRecorderStyle::StyleName("TakeRecorderStyle");
@@ -208,6 +209,14 @@ FTakeRecorderStyle::FTakeRecorderStyle()
 
 	Set("FontAwesome.28", FSlateFontInfo(FAppStyle::Get().GetFontStyle("FontAwesome.16").CompositeFont, 28));
 
+	// Hitch visualization
+	const FLinearColor ErrorColor(0.36f, 0.04f, 0.04f);
+	const FLinearColor WarningColor(0.31f, 0.31f, 0.03f);
+	Set("Hitching.ErrorColor", ErrorColor); 
+	Set("Hitching.WarningColor", WarningColor); 
+	Set("Hitching.CatchUpZoneBrush", new FSlateColorBrush(FLinearColor(WarningColor.R, WarningColor.G, WarningColor.B, 0.5f))); // time slider area where we were behind target timecode
+	Set("Hitching.MarkerWarningBrush", new IMAGE_BRUSH_SVG("AlertTriangleSolid", Icon16x16));
+
 	FSlateStyleRegistry::RegisterSlateStyle(*this);
 }
 
diff --git a/Engine/Plugins/VirtualProduction/Takes/Source/TakeRecorder/Private/Timecode/Visualization/HitchAnalysis.cpp b/Engine/Plugins/VirtualProduction/Takes/Source/TakeRecorder/Private/Timecode/Visualization/HitchAnalysis.cpp
new file mode 100644
--- /dev/null
+++ b/Engine/Plugins/VirtualProduction/Takes/Source/TakeRecorder/Private/Timecode/Visualization/HitchAnalysis.cpp
@@ -0,0 +1,210 @@
+﻿// Copyright Epic Games, Inc. All Rights Reserved.
+
+#include "HitchAnalysis.h"
+
+#include "Editor/Sequencer/Private/SSequencer.h"
+#include "Hitching/FrameHitchSceneDecoration.h"
+#include "MovieScene.h"
+
+namespace UE::TakeRecorder
+{
+namespace Private
+{
+static TAutoConsoleVariable<bool> CVarLogFrameMarkers(
+	TEXT("TakeRecorder.Hitching.LogFrameMarkers"),
+	false,
+	TEXT("Whether to log the actual and expected timecode for each frame marker that the hitching visualization shows.")
+	);
+static TAutoConsoleVariable<bool> CVarLogFrameCatchupRanges(
+	TEXT("TakeRecorder.Hitching.LogCatchupRanges"),
+	false,
+	TEXT("Whether to log start and end of catchup ranges the hitching visualization shows.")
+	);
+	
+template<typename TCallback> requires std::is_invocable_v<TCallback, const FFrameNumber& /*InCurrent*/, const FTimecode& /*InTarget*/, const FTimecode& /*InActual*/>
+static void ForEachFrame(UFrameHitchSceneDecoration& InDecoration, const FFrameRate& InRate, TCallback&& InCallback)
+{
+	const TConstArrayView<FFrameNumber> TargetTimes = InDecoration.TargetTimecode.GetFrameTimes();
+	const TConstArrayView<FFrameNumber> ActualTimes = InDecoration.ActualTimecode.GetFrameTimes();
+	// If this fires, the data probably was not recorded correctly / messed with. For each frame, we record target and actual time.
+	if (!ensureMsgf(TargetTimes.Num() == ActualTimes.Num(), TEXT("Recorded frame times do not match up")))
+	{
+		return;
+	}
+	
+	for (int32 Index = 0; Index < TargetTimes.Num(); ++Index)
+	{
+		const FFrameNumber Current = TargetTimes[Index];
+		// If this fires, the data probably was not recorded correctly / messed with. For each frame, we record target and actual time.
+		if (!ensureMsgf(ActualTimes[Index] == Current, TEXT("Recorded frame times do not match up")))
+		{
+			continue;
+		}
+		
+		const TOptional<FTimecode> TargetTimecode = InDecoration.TargetTimecode.Evaluate(Current);
+		const TOptional<FTimecode> ActualTimecode = InDecoration.ActualTimecode.Evaluate(Current);
+		if (ensure(TargetTimecode && ActualTimecode))
+		{
+			InCallback(Current, *TargetTimecode, *ActualTimecode);
+		}
+	}
+}
+	
+static FTimecode GetNextTimecode(const FTimecode& InCurrent, int32 InNumTimecodeFrames)
+{
+	FTimecode Next(InCurrent.Hours, InCurrent.Minutes, InCurrent.Seconds, InCurrent.Frames + 1, InCurrent.Subframe, InCurrent.bDropFrameFormat);
+	
+	if (Next.Frames >= InNumTimecodeFrames)
+	{
+		Next.Frames = 0;
+		++Next.Seconds;
+	}
+	
+	if (Next.Seconds >= 60)
+	{
+		Next.Seconds = 0;
+		++Next.Minutes;
+	}
+	
+	if (Next.Minutes >= 60)
+	{
+		Next.Minutes = 0;
+		++Next.Hours;
+	}
+	return Next;
+}
+
+struct FAnalyzedMarkers
+{
+	TArray<FSkippedTimecodeMarker> Skipped;
+	TArray<FSkippedTimecodeMarker> Repeated;
+};
+	
+static FAnalyzedMarkers AnalyzeSkippedTimecodes(
+	UFrameHitchSceneDecoration& InDecoration, const FFrameRate& InSequencerRate, int32 InNumTimecodeFrames
+	)
+{
+	FAnalyzedMarkers Result;
+
+	TOptional<FTimecode> LastTimecode;
+	ForEachFrame(InDecoration, InSequencerRate,
+		[&LastTimecode, InNumTimecodeFrames, &Result](const FFrameNumber& InCurrent, const FTimecode&, const FTimecode& InActual)
+		{
+			const TOptional<FTimecode> ExpectedTimecode = LastTimecode ? GetNextTimecode(*LastTimecode, InNumTimecodeFrames) : TOptional<FTimecode>{};
+			if (!ExpectedTimecode || InActual == *ExpectedTimecode)
+			{
+				return;
+			}
+
+			const bool bIsRepeat = LastTimecode && InActual == *LastTimecode;
+			TArray<FSkippedTimecodeMarker>& Markers = bIsRepeat ? Result.Repeated : Result.Skipped;
+			Markers.Add(FSkippedTimecodeMarker{ InCurrent });
+				
+			UE_CLOG(CVarLogFrameMarkers.GetValueOnAnyThread(), LogTemp, Warning, TEXT("%s timecode at %d. Actual: %s \tExpected: %s \tLast: %s"),
+				bIsRepeat ? TEXT("Repeated") : TEXT("Skipped"),InCurrent.Value, *InActual.ToString(), *ExpectedTimecode->ToString(),
+				LastTimecode ? *LastTimecode->ToString() : TEXT("unset")
+				);
+		
+			LastTimecode = InActual;
+		});
+	
+	return Result;
+}
+
+static void DumpCatchupRanges(UFrameHitchSceneDecoration& InDecoration, const TArray<FCatchupTimeRange>& InRanges)
+{
+	if (!CVarLogFrameCatchupRanges.GetValueOnAnyThread())
+	{
+		return;
+	}
+	
+	for (const FCatchupTimeRange& Range : InRanges)
+	{
+		const TOptional<FTimecode> Start_TargetFrameNumber = InDecoration.TargetTimecode.Evaluate(Range.StartTime);
+		const TOptional<FTimecode> Start_ActualFrameNumber = InDecoration.ActualTimecode.Evaluate(Range.StartTime);
+		const TOptional<FTimecode> End_TargetFrameNumber = InDecoration.TargetTimecode.Evaluate(Range.EndTime);
+		const TOptional<FTimecode> End_ActualFrameNumber = InDecoration.ActualTimecode.Evaluate(Range.EndTime);
+
+		if (!ensure(Start_TargetFrameNumber && Start_ActualFrameNumber && End_TargetFrameNumber && End_ActualFrameNumber))
+		{
+			return;
+		}
+		
+		UE_CLOG(CVarLogFrameCatchupRanges.GetValueOnAnyThread(), LogTemp, Warning,
+			TEXT("Catch-up range:\n\tStart at %d\tTarget TC: %s \tActual TC: %s\n\tEnd at %d   \tTarget TC: %s \tActual TC: %s"),
+			Range.StartTime.FrameNumber.Value, *Start_TargetFrameNumber->ToString(), *Start_ActualFrameNumber->ToString(),
+			Range.EndTime.FrameNumber.Value, *End_TargetFrameNumber->ToString(), *End_ActualFrameNumber->ToString()
+		);
+	}
+}
+
+static TArray<FCatchupTimeRange> AnalyzeCatchupZones(
+	UFrameHitchSceneDecoration& InDecoration, const FFrameRate& InSequencerRate
+	)
+{
+	TArray<FCatchupTimeRange> Result;
+	
+	TOptional<FCatchupTimeRange> OpenCatchupBracket;
+	ForEachFrame(InDecoration, InSequencerRate,
+		[&OpenCatchupBracket, &InSequencerRate, &Result](const FFrameNumber& InCurrent, const FTimecode& InTarget, const FTimecode& InActual)
+		{
+			const FFrameNumber TargetFrameNumber = InTarget.ToFrameNumber(InSequencerRate);
+			const FFrameNumber ActualFrameNumber = InActual.ToFrameNumber(InSequencerRate);
+			if (!OpenCatchupBracket && ActualFrameNumber < TargetFrameNumber)
+			{
+				OpenCatchupBracket = FCatchupTimeRange{ InCurrent, InCurrent };
+			}
+			else if (OpenCatchupBracket && ActualFrameNumber < TargetFrameNumber)
+			{
+				OpenCatchupBracket->EndTime = InCurrent;
+			}
+			else if (OpenCatchupBracket
+				// We don't expect actual to ever be higher than target, but we'll handle that case defensively;
+				// UCatchupFixedRateCustomTimeStep should have prevented the engine getting ahead of platform time.
+				&& ActualFrameNumber >= TargetFrameNumber) 
+			{
+				OpenCatchupBracket->FirstOkFrame = InCurrent;
+				
+				Result.Add(*OpenCatchupBracket);
+				OpenCatchupBracket.Reset();
+			}
+		});
+
+	// There may be a catch-up zone at the end where the engine never caught up
+	if (OpenCatchupBracket)
+	{
+		Result.Add(*OpenCatchupBracket);
+	}
+
+	DumpCatchupRanges(InDecoration, Result);
+	return Result;
+}
+}
+	
+TOptional<FTimecodeHitchData> AnalyseHitches(ISequencer& InSequencer)
+{
+	UMovieSceneSequence* RootSequence = InSequencer.GetRootMovieSceneSequence();
+	UMovieScene* MovieScene = RootSequence->GetMovieScene();
+	UFrameHitchSceneDecoration* Decoration = MovieScene->FindDecoration<UFrameHitchSceneDecoration>();
+	if (!Decoration)
+	{
+		return {};
+	}
+	
+	const FFrameRate Rate = InSequencer.GetFocusedTickResolution();
+	Private::FAnalyzedMarkers Markers = Private::AnalyzeSkippedTimecodes(*Decoration, Rate, Decoration->FrameRate.Numerator);
+	return FTimecodeHitchData
+	{
+		MoveTemp(Markers.Skipped), MoveTemp(Markers.Repeated),
+		Private::AnalyzeCatchupZones(*Decoration, Rate),
+	};
+}
+
+bool HasHitchData(ISequencer& InSequencer)
+{
+	UMovieSceneSequence* RootSequence = InSequencer.GetRootMovieSceneSequence();
+	UMovieScene* MovieScene = RootSequence->GetMovieScene();
+	UFrameHitchSceneDecoration* Decoration = MovieScene->FindDecoration<UFrameHitchSceneDecoration>();
+	return Decoration != nullptr;
+}
+}
\ No newline at end of file
diff --git a/Engine/Plugins/VirtualProduction/Takes/Source/TakeRecorder/Private/Timecode/Visualization/HitchAnalysis.h b/Engine/Plugins/VirtualProduction/Takes/Source/TakeRecorder/Private/Timecode/Visualization/HitchAnalysis.h
new file mode 100644
--- /dev/null
+++ b/Engine/Plugins/VirtualProduction/Takes/Source/TakeRecorder/Private/Timecode/Visualization/HitchAnalysis.h
@@ -0,0 +1,63 @@
+﻿// Copyright Epic Games, Inc. All Rights Reserved.
+
+#pragma once
+
+#include "Misc/FrameTime.h"
+#include "Misc/Optional.h"
+
+class ISequencer;
+
+namespace UE::TakeRecorder
+{
+/** A section of time where the engine fell behind and tried to catchup. */
+struct FCatchupTimeRange
+{
+	/** The first frame at which the timecode started being behind */
+	FFrameTime StartTime;
+	/** The last frame at which the timecode was still behind */
+	FFrameTime EndTime;
+
+	/**
+	 * The frame at which the timecode has caught up again.
+	 * That's the frame immediately after EndTime.
+	 *
+	 * Unset if the recording ended with a frame at which timecode was behind.
+	 */
+	TOptional<FFrameTime> FirstOkFrame;
+};
+
+/** Data about a frame at which timecode was skipped */
+struct FSkippedTimecodeMarker
+{
+	/** The frame at which the timecode did not match. */
+	FFrameTime Frame;
+};
+
+/** All data required to draw */
+struct FTimecodeHitchData
+{
+	/** Frames at which the timecode was skipped (previous frame was not current frame - 1). */
+	TArray<FSkippedTimecodeMarker> SkippedTimecodeMarkers;
+
+	/** Frames at which timecode was repeated (previous frame had the same timecode). */
+	TArray<FSkippedTimecodeMarker> RepeatedTimecodeMarkers;
+
+	/** Sections at which the engine fell behind timecode. */
+	TArray<FCatchupTimeRange> CatchupTimes;
+
+	void Reset()
+	{
+		SkippedTimecodeMarkers.Reset();
+		CatchupTimes.Reset();
+	}
+};
+
+/**
+ * Tries to find the required tracks for analysing hitches, and then analyses them.
+ * @return Analyzed hitch data if required tracks were found.
+ */
+TOptional<FTimecodeHitchData> AnalyseHitches(ISequencer& InSequencer);
+
+/** @return Whether there is any hitch data that can be analyzed. */
+bool HasHitchData(ISequencer& InSequencer);
+}
diff --git a/Engine/Plugins/VirtualProduction/Takes/Source/TakeRecorder/Private/Timecode/Visualization/HitchViewOptionUtils.cpp b/Engine/Plugins/VirtualProduction/Takes/Source/TakeRecorder/Private/Timecode/Visualization/HitchViewOptionUtils.cpp
new file mode 100644
--- /dev/null
+++ b/Engine/Plugins/VirtualProduction/Takes/Source/TakeRecorder/Private/Timecode/Visualization/HitchViewOptionUtils.cpp
@@ -0,0 +1,94 @@
+﻿// Copyright Epic Games, Inc. All Rights Reserved.
+
+#if WITH_EDITOR
+#include "HitchViewOptionUtils.h"
+
+#include "Framework/MultiBox/MultiBoxBuilder.h"
+#include "Framework/MultiBox/MultiBoxExtender.h"
+#include "Timecode/Visualization/Sequencer/HitchVisualizationViewModel.h"
+#include "Hitching/FrameHitchSceneDecoration.h"
+#include "MovieScene.h"
+#include "ScopedTransaction.h"
+#include "TakeRecorderCommands.h"
+#include "TakeRecorderHitchVisualizationSettings.h"
+#include "Templates/SharedPointer.h"
+
+#define LOCTEXT_NAMESPACE "FHitchVisualizerInSequencer"
+
+namespace UE::TakeRecorder
+{
+TSharedRef<FExtender> ExtendSequencerViewOptionsWithHitchVisualization()
+{
+	const TSharedRef<FExtender> Extender = MakeShared<FExtender>();
+	Extender->AddMenuExtension(TEXT("MarkedFrames"), EExtensionHook::After, nullptr,
+		FMenuExtensionDelegate::CreateLambda([](FMenuBuilder& MenuBuilder)
+		{
+			MenuBuilder.BeginSection(TEXT("HitchVisualization"), LOCTEXT("RecordingIntegrity", "Recording Integrity"));
+				
+			MenuBuilder.AddMenuEntry(
+				LOCTEXT("SkipMarkers.Title", "Show Frame Drop Markers"),
+				LOCTEXT("SkipMarkers.Description", "Whether to place markers on frames where a timecode frame was skipped."),
+				FSlateIcon(),
+				FUIAction(
+					FExecuteAction::CreateLambda([]{ UTakeRecorderHitchVisualizationSettings::Get()->ToggleShowFrameDropMarkers(); }),
+					FCanExecuteAction(),
+					FIsActionChecked::CreateLambda([]{ return UTakeRecorderHitchVisualizationSettings::Get()->GetShowFrameDropMarkers(); })
+					),
+					NAME_None, EUserInterfaceActionType::ToggleButton
+				);
+
+			MenuBuilder.AddMenuEntry(
+				LOCTEXT("RepeatMarkers.Title", "Show Frame Repeat Markers"),
+				LOCTEXT("RepeatMarkers.Description", "Whether to place markers on frames where a timecode frame was repeated."),
+				FSlateIcon(),
+				FUIAction(
+					FExecuteAction::CreateLambda([]{ UTakeRecorderHitchVisualizationSettings::Get()->ToggleShowFrameRepeatMarkers(); }),
+					FCanExecuteAction(),
+					FIsActionChecked::CreateLambda([]{ return UTakeRecorderHitchVisualizationSettings::Get()->GetShowFrameRepeatMarkers(); })
+					),
+					NAME_None, EUserInterfaceActionType::ToggleButton
+				);
+
+			MenuBuilder.AddMenuEntry(
+				LOCTEXT("ShowCatchupRanges.Title", "Show Hitch Recovery Ranges"),
+				LOCTEXT("ShowCatchupRanges.Description", "Whether to show areas in which the engine could not keep up, i.e. was running behind."),
+				FSlateIcon(),
+				FUIAction(
+					FExecuteAction::CreateLambda([]{ UTakeRecorderHitchVisualizationSettings::Get()->ToggleShowCatchupRanges(); }),
+					FCanExecuteAction(),
+					FIsActionChecked::CreateLambda([]{ return UTakeRecorderHitchVisualizationSettings::Get()->GetShowCatchupRanges(); })
+					),
+					NAME_None, EUserInterfaceActionType::ToggleButton
+				);
+
+			MenuBuilder.AddMenuEntry(FTakeRecorderCommands::Get().ClearRecordingIntegrityData);
+
+			MenuBuilder.EndSection();
+		}));
+	return Extender;
+}
+
+void Execute_ClearRecordingIntegrityData(TWeakPtr<ISequencer> InSequencer)
+{
+	const TSharedPtr<ISequencer> SequencerPin = InSequencer.Pin();
+	UMovieSceneSequence* Sequence = SequencerPin ? SequencerPin->GetRootMovieSceneSequence() : nullptr;
+	UMovieScene* MovieScene = Sequence ? Sequence->GetMovieScene() : nullptr;
+	if (MovieScene && !MovieScene->IsReadOnly())
+	{
+		const FScopedTransaction Transaction(LOCTEXT("ClearData", "Clear Recording Integrity Data"));
+		MovieScene->Modify();
+		MovieScene->RemoveDecoration<UFrameHitchSceneDecoration>();
+	}
+}
+
+bool CanExecute_ClearRecordingIntegrityData(TWeakPtr<ISequencer> InSequencer)
+{
+	const TSharedPtr<ISequencer> SequencerPin = InSequencer.Pin();
+	UMovieSceneSequence* Sequence = SequencerPin ? SequencerPin->GetRootMovieSceneSequence() : nullptr;
+	UMovieScene* MovieScene = Sequence ? Sequence->GetMovieScene() : nullptr;
+	return MovieScene && !MovieScene->IsReadOnly();
+}
+}
+
+#undef LOCTEXT_NAMESPACE
+#endif
\ No newline at end of file
diff --git a/Engine/Plugins/VirtualProduction/Takes/Source/TakeRecorder/Private/Timecode/Visualization/HitchViewOptionUtils.h b/Engine/Plugins/VirtualProduction/Takes/Source/TakeRecorder/Private/Timecode/Visualization/HitchViewOptionUtils.h
new file mode 100644
--- /dev/null
+++ b/Engine/Plugins/VirtualProduction/Takes/Source/TakeRecorder/Private/Timecode/Visualization/HitchViewOptionUtils.h
@@ -0,0 +1,30 @@
+﻿// Copyright Epic Games, Inc. All Rights Reserved.
+
+#pragma once
+
+#if WITH_EDITOR
+
+#include "Templates/SharedPointerFwd.h"
+
+class ISequencer;
+class FExtender;
+
+namespace UE::TakeRecorder
+{
+class FHitchVisualizationViewModel;
+	
+/**
+ * Builds an extender for the Sequencer view options adding:
+ * - Show Frame Drop Markers
+ * - Show Catchup Ranges
+ * - Clear Recording Integrity Markers and Ranges 
+ */
+TSharedRef<FExtender> ExtendSequencerViewOptionsWithHitchVisualization();
+
+/** Handles the FTakeRecorderCommands::Get().ClearRecordingIntegrityData. */
+void Execute_ClearRecordingIntegrityData(TWeakPtr<ISequencer> InSequencer);
+/** Whether FTakeRecorderCommands::Get().ClearRecordingIntegrityData can be executed, i.e. movie scene is not read-only. */
+bool CanExecute_ClearRecordingIntegrityData(TWeakPtr<ISequencer> InSequencer);
+}
+
+#endif
\ No newline at end of file
diff --git a/Engine/Plugins/VirtualProduction/Takes/Source/TakeRecorder/Private/Timecode/Visualization/Sequencer/HitchVisualizationViewModel.cpp b/Engine/Plugins/VirtualProduction/Takes/Source/TakeRecorder/Private/Timecode/Visualization/Sequencer/HitchVisualizationViewModel.cpp
new file mode 100644
--- /dev/null
+++ b/Engine/Plugins/VirtualProduction/Takes/Source/TakeRecorder/Private/Timecode/Visualization/Sequencer/HitchVisualizationViewModel.cpp
@@ -0,0 +1,48 @@
+﻿// Copyright Epic Games, Inc. All Rights Reserved.
+
+#include "HitchVisualizationViewModel.h"
+
+#include "Timecode/Visualization/TakeRecorderHitchVisualizationSettings.h"
+#include "Timecode/Visualization/VisualizationDrawingUtils.h"
+
+namespace UE::TakeRecorder
+{
+UE_SEQUENCER_DEFINE_CASTABLE(FHitchVisualizationViewModel);
+
+int32 FHitchVisualizationViewModel::OnPaintTimeSliderOverlay(
+	const FPaintArgs& Args, const FGeometry& AllottedGeometry, const FSlateRect& MyCullingRect, FSlateWindowElementList& OutDrawElements,
+	int32 LayerId, const FWidgetStyle& InWidgetStyle
+	)
+{
+	if (!CanPaint())
+	{
+		return LayerId;
+	}
+	
+	UTakeRecorderHitchVisualizationSettings* Settings = UTakeRecorderHitchVisualizationSettings::Get();
+	return DrawTimeSliderArea(
+		HitchData, MakeScrubRangeToScreen(AllottedGeometry), AllottedGeometry, OutDrawElements, LayerId, Settings->GetFlagsForUI()
+		);
+}
+
+int32 FHitchVisualizationViewModel::OnPaintTrackAreaOverlay(
+	const FPaintArgs& Args, const FGeometry& AllottedGeometry, const FSlateRect& MyCullingRect, FSlateWindowElementList& OutDrawElements,
+	int32 LayerId, const FWidgetStyle& InWidgetStyle)
+{
+	if (!CanPaint())
+	{
+		return LayerId;
+	}
+	
+	UTakeRecorderHitchVisualizationSettings* Settings = UTakeRecorderHitchVisualizationSettings::Get();
+	return DrawTrackArea(
+		HitchData, MakeScrubRangeToScreen(AllottedGeometry), AllottedGeometry, OutDrawElements, LayerId, Settings->GetFlagsForUI()
+		);
+}
+
+FScrubRangeToScreen FHitchVisualizationViewModel::MakeScrubRangeToScreen(const FGeometry& AllottedGeometry) const
+{
+	const TSharedPtr<ISequencer> SequencerPin = WeakOwningSequencer.Pin();
+	return FScrubRangeToScreen(SequencerPin->GetFocusedTickResolution(), SequencerPin->GetViewRange(), AllottedGeometry.Size);
+}
+}
diff --git a/Engine/Plugins/VirtualProduction/Takes/Source/TakeRecorder/Private/Timecode/Visualization/Sequencer/HitchVisualizationViewModel.h b/Engine/Plugins/VirtualProduction/Takes/Source/TakeRecorder/Private/Timecode/Visualization/Sequencer/HitchVisualizationViewModel.h
new file mode 100644
--- /dev/null
+++ b/Engine/Plugins/VirtualProduction/Takes/Source/TakeRecorder/Private/Timecode/Visualization/Sequencer/HitchVisualizationViewModel.h
@@ -0,0 +1,57 @@
+﻿// Copyright Epic Games, Inc. All Rights Reserved.
+
+#pragma once
+
+#if WITH_EDITOR
+#include "MVVM/Extensions/ITopTimeSliderOverlayExtension.h"
+#include "MVVM/Extensions/ITrackAreaOverlayExtension.h"
+#include "MVVM/ViewModels/ViewModel.h"
+#include "Timecode/Visualization/HitchAnalysis.h"
+#include "Timecode/Visualization/VisualizationDrawingUtils.h"
+
+namespace UE::TakeRecorder
+{
+/**  */
+class FHitchVisualizationViewModel
+	: public Sequencer::FViewModel
+	, public Sequencer::ITopTimeSliderOverlayExtension
+	, public Sequencer::ITrackAreaOverlayExtension
+{
+public:
+
+	UE_SEQUENCER_DECLARE_CASTABLE(FHitchVisualizationViewModel, ITopTimeSliderOverlayExtension, ITrackAreaOverlayExtension);
+
+	explicit FHitchVisualizationViewModel(
+		TWeakPtr<ISequencer> InWeakOwningSequencer, FTimecodeHitchData InHitchData, TAttribute<bool> InCanPaintAttr = true
+		)
+		: WeakOwningSequencer(MoveTemp(InWeakOwningSequencer))
+		, CanPaintAttr(MoveTemp(InCanPaintAttr))
+		, HitchData(MoveTemp(InHitchData))
+	{}
+	
+	//~ Begin ITimeSliderOverlayExtension Interface
+	virtual int32 OnPaintTimeSliderOverlay(const FPaintArgs& Args, const FGeometry& AllottedGeometry, const FSlateRect& MyCullingRect, FSlateWindowElementList& OutDrawElements, int32 LayerId, const FWidgetStyle& InWidgetStyle) override;
+	//~ End ITimeSliderOverlayExtension Interface
+	
+	//~ Begin ITrackAreaOverlayExtension Interface
+	virtual int32 OnPaintTrackAreaOverlay(const FPaintArgs& Args, const FGeometry& AllottedGeometry, const FSlateRect& MyCullingRect, FSlateWindowElementList& OutDrawElements, int32 LayerId, const FWidgetStyle& InWidgetStyle) override;
+	//~ End ITrackAreaOverlayExtension Interface
+
+private:
+
+	/** The sequencer that this model is saved in. */
+	const TWeakPtr<ISequencer> WeakOwningSequencer;
+
+	/** Decides whether we can paint. */
+	const TAttribute<bool> CanPaintAttr;
+
+	/** The hitch data that is being visualized. */
+	FTimecodeHitchData HitchData;
+
+	FScrubRangeToScreen MakeScrubRangeToScreen(const FGeometry& AllottedGeometry) const;
+
+	/** @return Whether this model should paint itself. */
+	bool CanPaint() const { return CanPaintAttr.Get(); }
+};
+}
+#endif
\ No newline at end of file
diff --git a/Engine/Plugins/VirtualProduction/Takes/Source/TakeRecorder/Private/Timecode/Visualization/Sequencer/SequencerHitchVisualizer.cpp b/Engine/Plugins/VirtualProduction/Takes/Source/TakeRecorder/Private/Timecode/Visualization/Sequencer/SequencerHitchVisualizer.cpp
new file mode 100644
--- /dev/null
+++ b/Engine/Plugins/VirtualProduction/Takes/Source/TakeRecorder/Private/Timecode/Visualization/Sequencer/SequencerHitchVisualizer.cpp
@@ -0,0 +1,85 @@
+﻿// Copyright Epic Games, Inc. All Rights Reserved.
+
+#if WITH_EDITOR
+#include "SequencerHitchVisualizer.h"
+
+#include "HitchVisualizationViewModel.h"
+#include "ISequencerModule.h"
+#include "MVVM/ViewModels/SequencerEditorViewModel.h"
+#include "MovieScene.h"
+#include "TakeRecorderCommands.h"
+#include "Timecode/Visualization/HitchViewOptionUtils.h"
+
+namespace UE::TakeRecorder
+{
+namespace Private
+{
+static void UnmapInjectedCommands(const TArray<TWeakPtr<ISequencer>>& InCustomizedSequencers)
+{
+	if (!FTakeRecorderCommands::IsRegistered())
+	{
+		return;	
+	}
+	for (const TWeakPtr<ISequencer>& WeakSequencer : InCustomizedSequencers)
+	{
+		const TSharedPtr<ISequencer> SequencerPin = WeakSequencer.Pin();
+		if (!SequencerPin)
+		{
+			continue;
+		}
+		
+		SequencerPin->GetCommandBindings()->UnmapAction(FTakeRecorderCommands::Get().ClearRecordingIntegrityData);
+	}
+}
+}
+	
+FSequencerHitchVisualizer::FSequencerHitchVisualizer()
+	: ViewOptionsExtender(TakeRecorder::ExtendSequencerViewOptionsWithHitchVisualization())
+{
+	ISequencerModule& SequencerModule = FModuleManager::Get().LoadModuleChecked<ISequencerModule>("Sequencer");
+	OnSequencerCreatedHandle = SequencerModule.RegisterOnSequencerCreated(
+		FOnSequencerCreated::FDelegate::CreateRaw(this, &FSequencerHitchVisualizer::OnSequencerCreated)
+		);
+	SequencerModule.GetViewMenuExtensibilityManager()->AddExtender(ViewOptionsExtender);
+}
+
+FSequencerHitchVisualizer::~FSequencerHitchVisualizer()
+{
+	if (ISequencerModule* SequencerModule = FModuleManager::Get().GetModulePtr<ISequencerModule>("Sequencer"))
+	{
+		SequencerModule->UnregisterOnSequencerCreated(OnSequencerCreatedHandle);
+		SequencerModule->GetViewMenuExtensibilityManager()->RemoveExtender(ViewOptionsExtender);
+
+		Private::UnmapInjectedCommands(CustomizedSequencers);
+	}
+}
+
+void FSequencerHitchVisualizer::OnSequencerCreated(TSharedRef<ISequencer> InSequencer)
+{
+	using namespace UE::Sequencer;
+
+	TOptional<FTimecodeHitchData> AnalyzedData = AnalyseHitches(*InSequencer);
+	if (!AnalyzedData)
+	{
+		return;
+	}
+
+	// To avoid implementing an undo command for undoing clearing the take data, we'll only draw the UI when the decorator is present.
+	const TAttribute<bool> CanPaint = TAttribute<bool>::CreateLambda([WeakSequencer = InSequencer.ToWeakPtr()]
+	{
+		const TSharedPtr<ISequencer> Sequencer = WeakSequencer.Pin();
+		return Sequencer && HasHitchData(*Sequencer);
+	});
+	const TSharedRef<FHitchVisualizationViewModel> Model = MakeShared<FHitchVisualizationViewModel>(InSequencer, *AnalyzedData, CanPaint);
+	InSequencer->GetViewModel()->GetEditorPanels().AddChild(Model);
+	CustomizedSequencers.Add(InSequencer);
+
+	InSequencer->GetCommandBindings()->MapAction(
+		FTakeRecorderCommands::Get().ClearRecordingIntegrityData,
+		FExecuteAction::CreateStatic(&TakeRecorder::Execute_ClearRecordingIntegrityData, InSequencer.ToWeakPtr()),
+		FCanExecuteAction::CreateStatic(&TakeRecorder::CanExecute_ClearRecordingIntegrityData, InSequencer.ToWeakPtr())
+		);
+}
+}
+
+#endif
\ No newline at end of file
diff --git a/Engine/Plugins/VirtualProduction/Takes/Source/TakeRecorder/Private/Timecode/Visualization/Sequencer/SequencerHitchVisualizer.h b/Engine/Plugins/VirtualProduction/Takes/Source/TakeRecorder/Private/Timecode/Visualization/Sequencer/SequencerHitchVisualizer.h
new file mode 100644
--- /dev/null
+++ b/Engine/Plugins/VirtualProduction/Takes/Source/TakeRecorder/Private/Timecode/Visualization/Sequencer/SequencerHitchVisualizer.h
@@ -0,0 +1,42 @@
+﻿// Copyright Epic Games, Inc. All Rights Reserved.
+
+#pragma once
+
+#include "Delegates/IDelegateInstance.h"
+#include "Templates/SharedPointer.h"
+#include "Templates/UnrealTemplate.h"
+
+#if WITH_EDITOR
+class FExtender;
+class ISequencer;
+
+namespace UE::TakeRecorder
+{
+class FHitchVisualizationViewModel;
+
+/** Hooks into Sequencer and extends the UI to visualize hitches. */
+class FSequencerHitchVisualizer : public FNoncopyable
+{
+public:
+	
+	FSequencerHitchVisualizer();
+	~FSequencerHitchVisualizer();
+
+private:
+
+	/** Handles for OnSequencerCreated. */
+	FDelegateHandle OnSequencerCreatedHandle;
+	
+	/** Extends Sequencer's view options */
+	const TSharedPtr<FExtender> ViewOptionsExtender;
+
+	/** The sequencers we have customized thus far. */
+	TArray<TWeakPtr<ISequencer>> CustomizedSequencers;
+
+	/** Invoked when a sequencer is created: extends the UI with visualization UI, if needed. */
+	void OnSequencerCreated(TSharedRef<ISequencer> InSequencer);
+};
+}
+#endif
+
+
diff --git a/Engine/Plugins/VirtualProduction/Takes/Source/TakeRecorder/Private/Timecode/Visualization/TakeRecorderHitchVisualizationSettings.cpp b/Engine/Plugins/VirtualProduction/Takes/Source/TakeRecorder/Private/Timecode/Visualization/TakeRecorderHitchVisualizationSettings.cpp
new file mode 100644
--- /dev/null
+++ b/Engine/Plugins/VirtualProduction/Takes/Source/TakeRecorder/Private/Timecode/Visualization/TakeRecorderHitchVisualizationSettings.cpp
@@ -0,0 +1,78 @@
+﻿// Copyright Epic Games, Inc. All Rights Reserved.
+
+#include "TakeRecorderHitchVisualizationSettings.h"
+
+#include "VisualizationDrawingUtils.h"
+
+bool UTakeRecorderHitchVisualizationSettings::GetShowFrameDropMarkers() const
+{
+	return bShowFrameDropMarkers;
+}
+void UTakeRecorderHitchVisualizationSettings::ToggleShowFrameDropMarkers()
+{
+	bShowFrameDropMarkers = !bShowFrameDropMarkers;
+	SaveConfig();
+}
+
+bool UTakeRecorderHitchVisualizationSettings::GetShowFrameRepeatMarkers() const
+{
+	return bShowFrameRepeatMarkers;
+}
+
+void UTakeRecorderHitchVisualizationSettings::ToggleShowFrameRepeatMarkers()
+{
+	bShowFrameRepeatMarkers = !bShowFrameRepeatMarkers;
+	SaveConfig();
+}
+
+bool UTakeRecorderHitchVisualizationSettings::GetShowCatchupRanges() const
+{
+	return bShowCatchupRanges;
+}
+void UTakeRecorderHitchVisualizationSettings::ToggleShowCatchupRanges()
+{
+	bShowCatchupRanges = !bShowCatchupRanges;
+	SaveConfig();
+}
+
+UE::TakeRecorder::EHitchDrawFlags UTakeRecorderHitchVisualizationSettings::GetFlagsForUI() const
+{
+	using namespace UE::TakeRecorder;
+	EHitchDrawFlags Flags = EHitchDrawFlags::None;
+
+	if (bShowCatchupRanges)
+	{
+		Flags |= EHitchDrawFlags::DrawCatchupRanges;
+	}
+
+	if (bShowFrameDropMarkers)
+	{
+		Flags |= EHitchDrawFlags::DrawSkippedTimecodeMarkers;
+	}
+
+	if (bShowFrameRepeatMarkers)
+	{
+		Flags |= EHitchDrawFlags::DrawRepeatedTimecodeMarkers;
+	}
+
+	return Flags;
+}
+
+void UTakeRecorderHitchVisualizationSettings::PostInitProperties()
+{
+	Super::PostInitProperties();
+	LoadConfig();
+}
+
+#if WITH_EDITOR
+void UTakeRecorderHitchVisualizationSettings::PostEditChangeProperty(FPropertyChangedEvent& PropertyChangedEvent)
+{
+	Super::PostEditChangeProperty(PropertyChangedEvent);
+
+	if (PropertyChangedEvent.ChangeType != EPropertyChangeType::Interactive)
+	{
+		// Dragging spinboxes causes this to be called every frame so we wait until they've finished dragging before saving.
+		SaveConfig();
+	}
+}
+#endif
\ No newline at end of file
diff --git a/Engine/Plugins/VirtualProduction/Takes/Source/TakeRecorder/Private/Timecode/Visualization/TakeRecorderHitchVisualizationSettings.h b/Engine/Plugins/VirtualProduction/Takes/Source/TakeRecorder/Private/Timecode/Visualization/TakeRecorderHitchVisualizationSettings.h
new file mode 100644
--- /dev/null
+++ b/Engine/Plugins/VirtualProduction/Takes/Source/TakeRecorder/Private/Timecode/Visualization/TakeRecorderHitchVisualizationSettings.h
@@ -0,0 +1,50 @@
+﻿// Copyright Epic Games, Inc. All Rights Reserved.
+
+#pragma once
+
+#include "UObject/Object.h"
+#include "TakeRecorderHitchVisualizationSettings.generated.h"
+
+namespace UE::TakeRecorder { enum class EHitchDrawFlags : uint8; }
+
+UCLASS(Config = EditorSettings)
+class UTakeRecorderHitchVisualizationSettings : public UObject
+{
+	GENERATED_BODY()
+public:
+
+	static UTakeRecorderHitchVisualizationSettings* Get()
+	{
+		return GetMutableDefault<UTakeRecorderHitchVisualizationSettings>();
+	}
+
+	bool GetShowFrameDropMarkers() const;
+	void ToggleShowFrameDropMarkers();
+
+	bool GetShowFrameRepeatMarkers() const;
+	void ToggleShowFrameRepeatMarkers();
+
+	bool GetShowCatchupRanges() const;
+	void ToggleShowCatchupRanges();
+
+	UE::TakeRecorder::EHitchDrawFlags GetFlagsForUI() const;
+
+	virtual void PostInitProperties() override;
+#if WITH_EDITOR
+	virtual void PostEditChangeProperty(FPropertyChangedEvent& PropertyChangedEvent) override;
+#endif
+	
+private:
+	
+	/** Whether to place markers on frames where a timecode frame was skipped. */
+	UPROPERTY(Config)
+	bool bShowFrameDropMarkers = true;
+
+	/** Whether to place markers on frames where a timecode frame was repeated. */
+	UPROPERTY(Config)
+	bool bShowFrameRepeatMarkers = true;
+
+	/** Whether to show areas in which the engine could not keep up, i.e. was running behind. */
+	UPROPERTY(Config)
+	bool bShowCatchupRanges = false;
+};
\ No newline at end of file
diff --git a/Engine/Plugins/VirtualProduction/Takes/Source/TakeRecorder/Private/Timecode/Visualization/VisualizationDrawingUtils.cpp b/Engine/Plugins/VirtualProduction/Takes/Source/TakeRecorder/Private/Timecode/Visualization/VisualizationDrawingUtils.cpp
new file mode 100644
--- /dev/null
+++ b/Engine/Plugins/VirtualProduction/Takes/Source/TakeRecorder/Private/Timecode/Visualization/VisualizationDrawingUtils.cpp
@@ -0,0 +1,202 @@
+﻿// Copyright Epic Games, Inc. All Rights Reserved.
+
+#include "VisualizationDrawingUtils.h"
+
+#include "HitchAnalysis.h"
+#include "Rendering/DrawElementTypes.h"
+#include "TakeRecorderStyle.h"
+
+namespace UE::TakeRecorder
+{
+namespace Private
+{
+const FVector2f MarkerBrushOffset{ 4.f, 4.f };
+	
+struct FDrawArgs
+{
+	const FScrubRangeToScreen& RangeToScreen;
+	const FGeometry& AllottedGeometry;
+	int32 LayerId;
+	
+	FSlateWindowElementList& OutDrawElements;
+
+	explicit FDrawArgs(const FScrubRangeToScreen& RangeToScreen, const FGeometry& AllottedGeometry, int32 LayerId, FSlateWindowElementList& OutDrawElements)
+		: RangeToScreen(RangeToScreen)
+		, AllottedGeometry(AllottedGeometry)
+		, LayerId(LayerId)
+		, OutDrawElements(OutDrawElements)
+	{}
+};
+
+static FGeometry MakeMarkerIconGeometry(
+	const FFrameTime& InTime, const FScrubRangeToScreen& InRangeToScreen, const FGeometry& InAllottedGeometry
+	)
+{
+	const float MarkerPos = InRangeToScreen.FrameToLocalX(InTime);
+	const FVector2f Position = FVector2f(MarkerPos, 0) + MarkerBrushOffset;
+	const FSlateBrush* Brush = FTakeRecorderStyle::Get().GetBrush("Hitching.MarkerWarningBrush");
+	const FVector2f Size = Brush->GetImageSize();
+	return InAllottedGeometry.MakeChild(Size, FSlateLayoutTransform(Position));
+}
+
+/**
+ * @return The frame at which to paint the range's right vertical line.
+ * 
+ * The line should be drawn at the beginning of the first frame after the range. Suppose we had a catchup range consisting of a single frame.
+ * Then StartTime and EndTime would be equal ... we wouldn't want to draw the left and right vertical lines at the same spot, i.e. a single line.
+ */
+static FFrameTime GetCatchupRangeLastVerticalLineFrame(const FCatchupTimeRange& InTimeRange)
+{
+	return InTimeRange.FirstOkFrame ? *InTimeRange.FirstOkFrame : InTimeRange.EndTime;
+}
+
+static FGeometry MakeCatchupZoneGeometry(
+	const FCatchupTimeRange& InTimeRange, const FScrubRangeToScreen& InRangeToScreen, const FGeometry& InAllottedGeometry
+	)
+{
+	const float Left = InRangeToScreen.FrameToLocalX(InTimeRange.StartTime);
+	const float Right = InRangeToScreen.FrameToLocalX(GetCatchupRangeLastVerticalLineFrame(InTimeRange));
+	const FVector2f Position(Left, 0);
+	const FVector2f Size(Right - Left, InAllottedGeometry.GetLocalSize().Y);
+	return InAllottedGeometry.MakeChild(Size, FSlateLayoutTransform(Position));
+}
+	
+static void DrawVerticalLine(
+	const FFrameTime& InFrameTime, const FDrawArgs& InDrawArgs,
+	const FLinearColor& InColor
+	)
+{
+	const float LinePos = InDrawArgs.RangeToScreen.FrameToLocalX(InFrameTime);
+	
+	TArray<FVector2D> LinePoints;
+	LinePoints.AddUninitialized(2);
+	LinePoints[0] = FVector2D(LinePos, 0.0f);
+	LinePoints[1] = FVector2D(LinePos, FMath::FloorToFloat(InDrawArgs.AllottedGeometry.Size.Y));
+
+	// Avoid painting over the tree view on the left, or the dockable tabs on the right.
+	const bool bOutOfBounds = LinePos < 0.f || LinePos >= InDrawArgs.AllottedGeometry.Size.X;
+	if (bOutOfBounds)
+	{
+		return;
+	}
+	
+	FSlateDrawElement::MakeLines(
+		InDrawArgs.OutDrawElements,
+		InDrawArgs.LayerId,
+		InDrawArgs.AllottedGeometry.ToPaintGeometry(),
+		LinePoints,
+		ESlateDrawEffect::None,
+		InColor,
+		false,
+		1.f
+	);
+}
+	
+static void DrawVerticalLines(TConstArrayView<FSkippedTimecodeMarker> InMarkers, const FDrawArgs& InDrawArgs)
+{
+	const FLinearColor& Color = FTakeRecorderStyle::Get().GetColor("Hitching.ErrorColor");
+	for (const FSkippedTimecodeMarker& Marker : InMarkers)
+	{
+		DrawVerticalLine(Marker.Frame, InDrawArgs, Color);
+	}
+}
+
+static void DrawVerticalLines(TConstArrayView<FCatchupTimeRange> InTimes, const FDrawArgs& InDrawArgs)
+{
+	const FLinearColor& Color = FTakeRecorderStyle::Get().GetColor("Hitching.WarningColor");
+	for (const FCatchupTimeRange& CatchupRange : InTimes)
+	{
+		DrawVerticalLine(CatchupRange.StartTime, InDrawArgs, Color);
+		DrawVerticalLine(GetCatchupRangeLastVerticalLineFrame(CatchupRange), InDrawArgs, Color);
+	}
+}
+	
+static void DrawMarkerWarningIcons(TConstArrayView<FSkippedTimecodeMarker> InMarkers, const FDrawArgs& InDrawArgs)
+{
+	for (const FSkippedTimecodeMarker& Marker : InMarkers)
+	{
+		const FSlateBrush* Brush = FTakeRecorderStyle::Get().GetBrush("Hitching.MarkerWarningBrush");
+		FSlateDrawElement::MakeBox(
+			InDrawArgs.OutDrawElements,
+			InDrawArgs.LayerId,
+			MakeMarkerIconGeometry(Marker.Frame, InDrawArgs.RangeToScreen, InDrawArgs.AllottedGeometry).ToPaintGeometry(),
+			Brush,
+			ESlateDrawEffect::None,
+			FTakeRecorderStyle::Get().GetColor("Hitching.ErrorColor")
+			);
+	}
+}
+
+static void DrawCatchupZones(TConstArrayView<FCatchupTimeRange> InTimes, const FDrawArgs& InDrawArgs)
+{
+	for (const FCatchupTimeRange& CatchupRange : InTimes)
+	{
+		FSlateDrawElement::MakeBox(
+			InDrawArgs.OutDrawElements,
+			InDrawArgs.LayerId,
+			MakeCatchupZoneGeometry(CatchupRange, InDrawArgs.RangeToScreen, InDrawArgs.AllottedGeometry).ToPaintGeometry(),
+			FTakeRecorderStyle::Get().GetBrush("Hitching.CatchUpZoneBrush"),
+			ESlateDrawEffect::None,
+			FTakeRecorderStyle::Get().GetColor("Hitching.WarningColor")
+			);
+	}
+}
+}
+	
+int32 DrawTimeSliderArea(
+	const FTimecodeHitchData& InData, const FScrubRangeToScreen& InRangeToScreen,
+	const FGeometry& InAllottedGeometry, FSlateWindowElementList& OutDrawElements, int32 InLayerId,
+	EHitchDrawFlags InFlags
+	)
+{
+	using namespace Private;
+	const FDrawArgs InDrawArgs(InRangeToScreen, InAllottedGeometry, InLayerId, OutDrawElements);
+	
+	if (EnumHasAnyFlags(InFlags, EHitchDrawFlags::DrawCatchupRanges))
+	{
+		DrawVerticalLines(InData.CatchupTimes, InDrawArgs);
+		DrawCatchupZones(InData.CatchupTimes, InDrawArgs);
+	}
+
+	if (EnumHasAnyFlags(InFlags, EHitchDrawFlags::DrawSkippedTimecodeMarkers))
+	{
+		DrawVerticalLines(InData.SkippedTimecodeMarkers, InDrawArgs);
+		DrawMarkerWarningIcons(InData.SkippedTimecodeMarkers, InDrawArgs);
+	}
+	
+	if (EnumHasAnyFlags(InFlags, EHitchDrawFlags::DrawRepeatedTimecodeMarkers))
+	{
+		DrawVerticalLines(InData.RepeatedTimecodeMarkers, InDrawArgs);
+		DrawMarkerWarningIcons(InData.RepeatedTimecodeMarkers, InDrawArgs);
+	}
+	
+	return InLayerId + 1;
+}
+
+int32 DrawTrackArea(
+	const FTimecodeHitchData& InData, const FScrubRangeToScreen& InRangeToScreen,
+	const FGeometry& InAllottedGeometry, FSlateWindowElementList& OutDrawElements, int32 InLayerId,
+	EHitchDrawFlags InFlags
+	)
+{
+	using namespace Private;
+	const FDrawArgs InDrawArgs(InRangeToScreen, InAllottedGeometry, InLayerId, OutDrawElements);
+	
+	if (EnumHasAnyFlags(InFlags, EHitchDrawFlags::DrawCatchupRanges))
+	{
+		DrawVerticalLines(InData.CatchupTimes, InDrawArgs);
+	}
+	
+	if (EnumHasAnyFlags(InFlags, EHitchDrawFlags::DrawSkippedTimecodeMarkers))
+	{
+		DrawVerticalLines(InData.SkippedTimecodeMarkers, InDrawArgs);
+	}
+
+	if (EnumHasAnyFlags(InFlags, EHitchDrawFlags::DrawRepeatedTimecodeMarkers))
+	{
+		DrawVerticalLines(InData.RepeatedTimecodeMarkers, InDrawArgs);
+	}
+	
+	return InLayerId + 1;
+}
+}
\ No newline at end of file
diff --git a/Engine/Plugins/VirtualProduction/Takes/Source/TakeRecorder/Private/Timecode/Visualization/VisualizationDrawingUtils.h b/Engine/Plugins/VirtualProduction/Takes/Source/TakeRecorder/Private/Timecode/Visualization/VisualizationDrawingUtils.h
new file mode 100644
--- /dev/null
+++ b/Engine/Plugins/VirtualProduction/Takes/Source/TakeRecorder/Private/Timecode/Visualization/VisualizationDrawingUtils.h
@@ -0,0 +1,58 @@
+﻿// Copyright Epic Games, Inc. All Rights Reserved.
+
+#pragma once
+
+#include "SequencerTimeSliderController.h"
+#include "Containers/Array.h"
+#include "Misc/FrameTime.h"
+
+struct FGeometry;
+class FSlateWindowElementList;
+
+namespace UE::TakeRecorder
+{
+struct FTimecodeHitchData;
+
+struct FScrubRangeToScreen : FSequencerTimeSliderController::FScrubRangeToScreen
+{
+	const FFrameRate FrameRate;
+	
+	explicit FScrubRangeToScreen(const FFrameRate& InFrameRate, const TRange<double>& InViewInput, const FVector2D& InWidgetSize)
+		: FSequencerTimeSliderController::FScrubRangeToScreen(InViewInput, InWidgetSize)
+		, FrameRate(InFrameRate)
+	{}
+
+	/** FrameTime -> local Widget Space */
+	float FrameToLocalX(const FFrameTime& InFrameTime) const
+	{
+		const double Seconds = InFrameTime.FrameNumber / FrameRate; 
+		return InputToLocalX(Seconds);
+	}
+};
+
+enum class EHitchDrawFlags : uint8
+{
+	None = 0,
+	
+	DrawSkippedTimecodeMarkers = 1 << 0,
+	DrawRepeatedTimecodeMarkers = 1 << 1,
+	DrawCatchupRanges = 1 << 2,
+	
+	All = DrawSkippedTimecodeMarkers | DrawRepeatedTimecodeMarkers | DrawCatchupRanges
+};
+ENUM_CLASS_FLAGS(EHitchDrawFlags);
+	
+/** Draws the top part (markers, etc.) of InData. */
+int32 DrawTimeSliderArea(
+	const FTimecodeHitchData& InData, const FScrubRangeToScreen& InRangeToScreen,
+	const FGeometry& InAllottedGeometry, FSlateWindowElementList& OutDrawElements, int32 InLayerId,
+	EHitchDrawFlags InFlags = EHitchDrawFlags::All
+	);
+	
+/** Draws the bottom part (markers, etc.) of InData. */
+int32 DrawTrackArea(
+	const FTimecodeHitchData& InData, const FScrubRangeToScreen& InRangeToScreen,
+	const FGeometry& InAllottedGeometry, FSlateWindowElementList& OutDrawElements, int32 InLayerId,
+	EHitchDrawFlags InFlags = EHitchDrawFlags::All
+	);
+}
diff --git a/Engine/Plugins/VirtualProduction/Takes/Source/TakeRecorder/TakeRecorder.Build.cs b/Engine/Plugins/VirtualProduction/Takes/Source/TakeRecorder/TakeRecorder.Build.cs
--- a/Engine/Plugins/VirtualProduction/Takes/Source/TakeRecorder/TakeRecorder.Build.cs
+++ b/Engine/Plugins/VirtualProduction/Takes/Source/TakeRecorder/TakeRecorder.Build.cs
@@ -37,7 +37,8 @@ public class TakeRecorder : ModuleRules
 				"Analytics",
 				"ToolMenus",
 				"ToolWidgets",
-				"NamingTokens"
+				"NamingTokens",
+				"SequencerCore",
 			}
 		);
 
