Take Recorder: Fix skip & repeat markers not showing up.

diff --git a/Engine/Plugins/VirtualProduction/Takes/Source/TakeRecorder/Private/Timecode/Visualization/HitchAnalysis.cpp b/Engine/Plugins/VirtualProduction/Takes/Source/TakeRecorder/Private/Timecode/Visualization/HitchAnalysis.cpp
--- a/Engine/Plugins/VirtualProduction/Takes/Source/TakeRecorder/Private/Timecode/Visualization/HitchAnalysis.cpp
+++ b/Engine/Plugins/VirtualProduction/Takes/Source/TakeRecorder/Private/Timecode/Visualization/HitchAnalysis.cpp
@@ -4,6 +4,7 @@
 
 #include "Editor/Sequencer/Private/SSequencer.h"
 #include "Hitching/FrameHitchSceneDecoration.h"
+#include "Misc/ScopeExit.h"
 #include "MovieScene.h"
 
 namespace UE::TakeRecorder
@@ -91,6 +92,8 @@ static FAnalyzedMarkers AnalyzeSkippedTimecodes(
 		[&LastTimecode, InNumTimecodeFrames, &Result](const FFrameNumber& InCurrent, const FTimecode&, const FTimecode& InActual)
 		{
 			const TOptional<FTimecode> ExpectedTimecode = LastTimecode ? GetNextTimecode(*LastTimecode, InNumTimecodeFrames) : TOptional<FTimecode>{};
+			ON_SCOPE_EXIT { LastTimecode = InActual; };
+				
 			if (!ExpectedTimecode || InActual == *ExpectedTimecode)
 			{
 				return;
@@ -104,8 +107,6 @@ static FAnalyzedMarkers AnalyzeSkippedTimecodes(
 				bIsRepeat ? TEXT("Repeated") : TEXT("Skipped"),InCurrent.Value, *InActual.ToString(), *ExpectedTimecode->ToString(),
 				LastTimecode ? *LastTimecode->ToString() : TEXT("unset")
 				);
-		
-			LastTimecode = InActual;
 		});
 	
 	return Result;
