Allow take recorder microphone sources to work in -game.

[RN] Using take recording feature, you can now capture microphone audio sources.

#DisabledValidations NonUnityAndStaticAnalysis

diff --git a/Engine/Plugins/Runtime/AudioCapture/AudioCapture.uplugin b/Engine/Plugins/Runtime/AudioCapture/AudioCapture.uplugin
--- a/Engine/Plugins/Runtime/AudioCapture/AudioCapture.uplugin
+++ b/Engine/Plugins/Runtime/AudioCapture/AudioCapture.uplugin
@@ -25,7 +25,7 @@
  		},
 		{
 			"Name": "AudioCaptureEditor",
-			"Type": "Editor",
+			"Type": "UncookedOnly",
 			"ProgramAllowList": [ "LiveLinkHub" ],
 			"LoadingPhase": "PostEngineInit",
 			"PlatformAllowList": [ "Win64" ]
diff --git a/Engine/Plugins/VirtualProduction/Takes/Source/TakeRecorderSources/Private/TakeRecorderMicrophoneAudioManager.cpp b/Engine/Plugins/VirtualProduction/Takes/Source/TakeRecorderSources/Private/TakeRecorderMicrophoneAudioManager.cpp
--- a/Engine/Plugins/VirtualProduction/Takes/Source/TakeRecorderSources/Private/TakeRecorderMicrophoneAudioManager.cpp
+++ b/Engine/Plugins/VirtualProduction/Takes/Source/TakeRecorderSources/Private/TakeRecorderMicrophoneAudioManager.cpp
@@ -59,6 +59,12 @@ void UTakeRecorderMicrophoneAudioManager::StartRecording(int32 InChannelCount)
 		return;
 	}
 
+	if (!GIsEditor && !IsInAudioThread())
+	{
+		UE_LOG(LogTakesCore, Error, TEXT("Microphone Audio Source will not start. Running in game with audio thread. Audio thread is not supported with take recording set UseAudioThread=false in DefaultEngine.ini to record audio in -game."));
+		return;
+	}
+	
 	if (AudioRecorder.IsValid())
 	{
 		if (AudioRecorder->IsReadyToRecord())
diff --git a/Engine/Plugins/VirtualProduction/Takes/Source/TakeRecorderSources/Private/TakeRecorderMicrophoneAudioSource.cpp b/Engine/Plugins/VirtualProduction/Takes/Source/TakeRecorderSources/Private/TakeRecorderMicrophoneAudioSource.cpp
--- a/Engine/Plugins/VirtualProduction/Takes/Source/TakeRecorderSources/Private/TakeRecorderMicrophoneAudioSource.cpp
+++ b/Engine/Plugins/VirtualProduction/Takes/Source/TakeRecorderSources/Private/TakeRecorderMicrophoneAudioSource.cpp
@@ -79,11 +79,13 @@ UTakeRecorderMicrophoneAudioSource::UTakeRecorderMicrophoneAudioSource(const FOb
 {
 }
 
-void UTakeRecorderMicrophoneAudioSource::Initialize()
+void UTakeRecorderMicrophoneAudioSource::DoAudioSetup()
 {
 	UTakeRecorderMicrophoneAudioManager* AudioInputManager = GetAudioInputManager();
-	if (AudioInputManager != nullptr)
+	if (!bHaveDoneAudioSetup && AudioInputManager != nullptr)
 	{
+		AudioInputManager->EnumerateAudioDevices(true);
+	
 		TUniquePtr<TArray<bool>> ChannelsInUse = GetChannelsInUse(AudioInputManager->GetDeviceChannelCount());
 		int32 AvailableChannelNumber = INDEX_NONE;
 
@@ -102,11 +104,19 @@ void UTakeRecorderMicrophoneAudioSource::Initialize()
 		}
 
 		AudioInputManager->GetOnNotifySourcesOfDeviceChange().AddUObject(this, &UTakeRecorderMicrophoneAudioSource::OnNotifySourcesOfDeviceChange);
+
+		bHaveDoneAudioSetup = true;
 	}
 }
 
+void UTakeRecorderMicrophoneAudioSource::Initialize()
+{
+	DoAudioSetup();
+}
+
 void UTakeRecorderMicrophoneAudioSource::PostEditChangeProperty(FPropertyChangedEvent& PropertyChangedEvent)
 {
+	DoAudioSetup();
 	Super::PostEditChangeProperty(PropertyChangedEvent);
 }
 
diff --git a/Engine/Plugins/VirtualProduction/Takes/Source/TakeRecorderSources/Private/TakeRecorderMicrophoneAudioSource.h b/Engine/Plugins/VirtualProduction/Takes/Source/TakeRecorderSources/Private/TakeRecorderMicrophoneAudioSource.h
--- a/Engine/Plugins/VirtualProduction/Takes/Source/TakeRecorderSources/Private/TakeRecorderMicrophoneAudioSource.h
+++ b/Engine/Plugins/VirtualProduction/Takes/Source/TakeRecorderSources/Private/TakeRecorderMicrophoneAudioSource.h
@@ -156,11 +156,13 @@ private:
 
 	/** Evaluate our tokens and automatically create an appropriate context. */
 	UE_API bool TryEvaluateTokenString(const FString& InString, const ULevelSequence* InSequence, FString& OutString) const;
-	
-private:
 
+private:
+	void DoAudioSetup();
+		
 	// Holds the Sequencer audio track for this source
 	TWeakObjectPtr<class UMovieSceneAudioTrack> CachedAudioTrack;
+
 	// Holds the USoundWave asset for a given take
 	TWeakObjectPtr<USoundWave> RecordedSoundWave;
 	// Caches the starting timecode for this take so it can be referenced when creating USoundWave assets
@@ -168,8 +170,11 @@ private:
 
 	// The user specified directory to store recorded audio assets in
 	FDirectoryPath AudioDirectory;
-	// The name of the UsoundWave asset 
+	// The name of the UsoundWave asset
 	FString AssetFileName;
+
+	// Have we initialized our source.
+	bool bHaveDoneAudioSetup = false;
 };
 
 #undef UE_API
