Take Recorder: Move hitch drawing logic into separate files.

diff --git a/Engine/Plugins/VirtualProduction/Takes/Source/TakeRecorder/Private/Timecode/Visualization/Drawing/EHitchDrawFlags.h b/Engine/Plugins/VirtualProduction/Takes/Source/TakeRecorder/Private/Timecode/Visualization/Drawing/EHitchDrawFlags.h
new file mode 100644
--- /dev/null
+++ b/Engine/Plugins/VirtualProduction/Takes/Source/TakeRecorder/Private/Timecode/Visualization/Drawing/EHitchDrawFlags.h
@@ -0,0 +1,21 @@
+﻿// Copyright Epic Games, Inc. All Rights Reserved.
+
+#pragma once
+
+#include "HAL/Platform.h"
+#include "Misc/EnumClassFlags.h"
+
+namespace UE::TakeRecorder
+{
+enum class EHitchDrawFlags : uint8
+{
+	None = 0,
+
+	DrawSkippedTimecodeMarkers = 1 << 0,
+	DrawRepeatedTimecodeMarkers = 1 << 1,
+	DrawCatchupRanges = 1 << 2,
+
+	All = DrawSkippedTimecodeMarkers | DrawRepeatedTimecodeMarkers | DrawCatchupRanges
+};
+ENUM_CLASS_FLAGS(EHitchDrawFlags);
+}
diff --git a/Engine/Plugins/VirtualProduction/Takes/Source/TakeRecorder/Private/Timecode/Visualization/Drawing/HitchHoverInfo.h b/Engine/Plugins/VirtualProduction/Takes/Source/TakeRecorder/Private/Timecode/Visualization/Drawing/HitchHoverInfo.h
new file mode 100644
--- /dev/null
+++ b/Engine/Plugins/VirtualProduction/Takes/Source/TakeRecorder/Private/Timecode/Visualization/Drawing/HitchHoverInfo.h
@@ -0,0 +1,11 @@
+﻿// Copyright Epic Games, Inc. All Rights Reserved.
+
+#pragma once
+
+namespace UE::TakeRecorder
+{
+struct FHitchHoverInfo
+{
+	
+};
+}
diff --git a/Engine/Plugins/VirtualProduction/Takes/Source/TakeRecorder/Private/Timecode/Visualization/Drawing/ScrubRangeToScreen.h b/Engine/Plugins/VirtualProduction/Takes/Source/TakeRecorder/Private/Timecode/Visualization/Drawing/ScrubRangeToScreen.h
new file mode 100644
--- /dev/null
+++ b/Engine/Plugins/VirtualProduction/Takes/Source/TakeRecorder/Private/Timecode/Visualization/Drawing/ScrubRangeToScreen.h
@@ -0,0 +1,28 @@
+﻿// Copyright Epic Games, Inc. All Rights Reserved.
+
+#pragma once
+
+#include "SequencerTimeSliderController.h"
+#include "Misc/FrameTime.h"
+#include "Misc/FrameRate.h"
+
+namespace UE::TakeRecorder
+{
+/** Helps convert FFrameTime to coordinates in widget space for finding the X-coordinate of vertical lines in Sequencer. */
+struct FScrubRangeToScreen : FSequencerTimeSliderController::FScrubRangeToScreen
+{
+	const FFrameRate FrameRate;
+
+	explicit FScrubRangeToScreen(const FFrameRate& InFrameRate, const TRange<double>& InViewInput, const FVector2D& InWidgetSize)
+		: FSequencerTimeSliderController::FScrubRangeToScreen(InViewInput, InWidgetSize)
+		, FrameRate(InFrameRate)
+	{}
+
+	/** FrameTime -> local Widget Space */
+	float FrameToLocalX(const FFrameTime& InFrameTime) const
+	{
+		const double Seconds = InFrameTime.FrameNumber / FrameRate; 
+		return InputToLocalX(Seconds);
+	}
+};
+}
diff --git a/Engine/Plugins/VirtualProduction/Takes/Source/TakeRecorder/Private/Timecode/Visualization/Drawing/VisualizationDrawingUtils.cpp b/Engine/Plugins/VirtualProduction/Takes/Source/TakeRecorder/Private/Timecode/Visualization/Drawing/VisualizationDrawingUtils.cpp
new file mode 100644
--- /dev/null
+++ b/Engine/Plugins/VirtualProduction/Takes/Source/TakeRecorder/Private/Timecode/Visualization/Drawing/VisualizationDrawingUtils.cpp
@@ -0,0 +1,204 @@
+﻿// Copyright Epic Games, Inc. All Rights Reserved.
+
+#include "VisualizationDrawingUtils.h"
+
+#include "Misc/FrameTime.h"
+#include "Rendering/DrawElementTypes.h"
+#include "ScrubRangeToScreen.h"
+#include "TakeRecorderStyle.h"
+#include "Timecode/Visualization/HitchAnalysis.h"
+
+namespace UE::TakeRecorder
+{
+namespace Private
+{
+const FVector2f MarkerBrushOffset{ 4.f, 4.f };
+	
+struct FDrawArgs
+{
+	const FScrubRangeToScreen& RangeToScreen;
+	const FGeometry& AllottedGeometry;
+	int32 LayerId;
+	
+	FSlateWindowElementList& OutDrawElements;
+
+	explicit FDrawArgs(const FScrubRangeToScreen& RangeToScreen, const FGeometry& AllottedGeometry, int32 LayerId, FSlateWindowElementList& OutDrawElements)
+		: RangeToScreen(RangeToScreen)
+		, AllottedGeometry(AllottedGeometry)
+		, LayerId(LayerId)
+		, OutDrawElements(OutDrawElements)
+	{}
+};
+
+static FGeometry MakeMarkerIconGeometry(
+	const FFrameTime& InTime, const FScrubRangeToScreen& InRangeToScreen, const FGeometry& InAllottedGeometry
+	)
+{
+	const float MarkerPos = InRangeToScreen.FrameToLocalX(InTime);
+	const FVector2f Position = FVector2f(MarkerPos, 0) + MarkerBrushOffset;
+	const FSlateBrush* Brush = FTakeRecorderStyle::Get().GetBrush("Hitching.MarkerWarningBrush");
+	const FVector2f Size = Brush->GetImageSize();
+	return InAllottedGeometry.MakeChild(Size, FSlateLayoutTransform(Position));
+}
+
+/**
+ * @return The frame at which to paint the range's right vertical line.
+ * 
+ * The line should be drawn at the beginning of the first frame after the range. Suppose we had a catchup range consisting of a single frame.
+ * Then StartTime and EndTime would be equal ... we wouldn't want to draw the left and right vertical lines at the same spot, i.e. a single line.
+ */
+static FFrameTime GetCatchupRangeLastVerticalLineFrame(const FCatchupTimeRange& InTimeRange)
+{
+	return InTimeRange.FirstOkFrame ? *InTimeRange.FirstOkFrame : InTimeRange.EndTime;
+}
+
+static FGeometry MakeCatchupZoneGeometry(
+	const FCatchupTimeRange& InTimeRange, const FScrubRangeToScreen& InRangeToScreen, const FGeometry& InAllottedGeometry
+	)
+{
+	const float Left = InRangeToScreen.FrameToLocalX(InTimeRange.StartTime);
+	const float Right = InRangeToScreen.FrameToLocalX(GetCatchupRangeLastVerticalLineFrame(InTimeRange));
+	const FVector2f Position(Left, 0);
+	const FVector2f Size(Right - Left, InAllottedGeometry.GetLocalSize().Y);
+	return InAllottedGeometry.MakeChild(Size, FSlateLayoutTransform(Position));
+}
+	
+static void DrawVerticalLine(
+	const FFrameTime& InFrameTime, const FDrawArgs& InDrawArgs,
+	const FLinearColor& InColor
+	)
+{
+	const float LinePos = InDrawArgs.RangeToScreen.FrameToLocalX(InFrameTime);
+	
+	TArray<FVector2D> LinePoints;
+	LinePoints.AddUninitialized(2);
+	LinePoints[0] = FVector2D(LinePos, 0.0f);
+	LinePoints[1] = FVector2D(LinePos, FMath::FloorToFloat(InDrawArgs.AllottedGeometry.Size.Y));
+
+	// Avoid painting over the tree view on the left, or the dockable tabs on the right.
+	const bool bOutOfBounds = LinePos < 0.f || LinePos >= InDrawArgs.AllottedGeometry.Size.X;
+	if (bOutOfBounds)
+	{
+		return;
+	}
+	
+	FSlateDrawElement::MakeLines(
+		InDrawArgs.OutDrawElements,
+		InDrawArgs.LayerId,
+		InDrawArgs.AllottedGeometry.ToPaintGeometry(),
+		LinePoints,
+		ESlateDrawEffect::None,
+		InColor,
+		false,
+		1.f
+	);
+}
+	
+static void DrawVerticalLines(TConstArrayView<FSkippedTimecodeMarker> InMarkers, const FDrawArgs& InDrawArgs)
+{
+	const FLinearColor& Color = FTakeRecorderStyle::Get().GetColor("Hitching.ErrorColor");
+	for (const FSkippedTimecodeMarker& Marker : InMarkers)
+	{
+		DrawVerticalLine(Marker.Frame, InDrawArgs, Color);
+	}
+}
+
+static void DrawVerticalLines(TConstArrayView<FCatchupTimeRange> InTimes, const FDrawArgs& InDrawArgs)
+{
+	const FLinearColor& Color = FTakeRecorderStyle::Get().GetColor("Hitching.WarningColor");
+	for (const FCatchupTimeRange& CatchupRange : InTimes)
+	{
+		DrawVerticalLine(CatchupRange.StartTime, InDrawArgs, Color);
+		DrawVerticalLine(GetCatchupRangeLastVerticalLineFrame(CatchupRange), InDrawArgs, Color);
+	}
+}
+	
+static void DrawMarkerWarningIcons(TConstArrayView<FSkippedTimecodeMarker> InMarkers, const FDrawArgs& InDrawArgs)
+{
+	for (const FSkippedTimecodeMarker& Marker : InMarkers)
+	{
+		const FSlateBrush* Brush = FTakeRecorderStyle::Get().GetBrush("Hitching.MarkerWarningBrush");
+		FSlateDrawElement::MakeBox(
+			InDrawArgs.OutDrawElements,
+			InDrawArgs.LayerId,
+			MakeMarkerIconGeometry(Marker.Frame, InDrawArgs.RangeToScreen, InDrawArgs.AllottedGeometry).ToPaintGeometry(),
+			Brush,
+			ESlateDrawEffect::None,
+			FTakeRecorderStyle::Get().GetColor("Hitching.ErrorColor")
+			);
+	}
+}
+
+static void DrawCatchupZones(TConstArrayView<FCatchupTimeRange> InTimes, const FDrawArgs& InDrawArgs)
+{
+	for (const FCatchupTimeRange& CatchupRange : InTimes)
+	{
+		FSlateDrawElement::MakeBox(
+			InDrawArgs.OutDrawElements,
+			InDrawArgs.LayerId,
+			MakeCatchupZoneGeometry(CatchupRange, InDrawArgs.RangeToScreen, InDrawArgs.AllottedGeometry).ToPaintGeometry(),
+			FTakeRecorderStyle::Get().GetBrush("Hitching.CatchUpZoneBrush"),
+			ESlateDrawEffect::None,
+			FTakeRecorderStyle::Get().GetColor("Hitching.WarningColor")
+			);
+	}
+}
+}
+	
+int32 DrawTimeSliderArea(
+	const FTimecodeHitchData& InData, const FScrubRangeToScreen& InRangeToScreen,
+	const FGeometry& InAllottedGeometry, FSlateWindowElementList& OutDrawElements, int32 InLayerId,
+	EHitchDrawFlags InFlags
+	)
+{
+	using namespace Private;
+	const FDrawArgs InDrawArgs(InRangeToScreen, InAllottedGeometry, InLayerId, OutDrawElements);
+	
+	if (EnumHasAnyFlags(InFlags, EHitchDrawFlags::DrawCatchupRanges))
+	{
+		DrawVerticalLines(InData.CatchupTimes, InDrawArgs);
+		DrawCatchupZones(InData.CatchupTimes, InDrawArgs);
+	}
+
+	if (EnumHasAnyFlags(InFlags, EHitchDrawFlags::DrawSkippedTimecodeMarkers))
+	{
+		DrawVerticalLines(InData.SkippedTimecodeMarkers, InDrawArgs);
+		DrawMarkerWarningIcons(InData.SkippedTimecodeMarkers, InDrawArgs);
+	}
+	
+	if (EnumHasAnyFlags(InFlags, EHitchDrawFlags::DrawRepeatedTimecodeMarkers))
+	{
+		DrawVerticalLines(InData.RepeatedTimecodeMarkers, InDrawArgs);
+		DrawMarkerWarningIcons(InData.RepeatedTimecodeMarkers, InDrawArgs);
+	}
+	
+	return InLayerId + 1;
+}
+
+int32 DrawTrackArea(
+	const FTimecodeHitchData& InData, const FScrubRangeToScreen& InRangeToScreen,
+	const FGeometry& InAllottedGeometry, FSlateWindowElementList& OutDrawElements, int32 InLayerId,
+	EHitchDrawFlags InFlags
+	)
+{
+	using namespace Private;
+	const FDrawArgs InDrawArgs(InRangeToScreen, InAllottedGeometry, InLayerId, OutDrawElements);
+	
+	if (EnumHasAnyFlags(InFlags, EHitchDrawFlags::DrawCatchupRanges))
+	{
+		DrawVerticalLines(InData.CatchupTimes, InDrawArgs);
+	}
+	
+	if (EnumHasAnyFlags(InFlags, EHitchDrawFlags::DrawSkippedTimecodeMarkers))
+	{
+		DrawVerticalLines(InData.SkippedTimecodeMarkers, InDrawArgs);
+	}
+
+	if (EnumHasAnyFlags(InFlags, EHitchDrawFlags::DrawRepeatedTimecodeMarkers))
+	{
+		DrawVerticalLines(InData.RepeatedTimecodeMarkers, InDrawArgs);
+	}
+	
+	return InLayerId + 1;
+}
+}
diff --git a/Engine/Plugins/VirtualProduction/Takes/Source/TakeRecorder/Private/Timecode/Visualization/Drawing/VisualizationDrawingUtils.h b/Engine/Plugins/VirtualProduction/Takes/Source/TakeRecorder/Private/Timecode/Visualization/Drawing/VisualizationDrawingUtils.h
new file mode 100644
--- /dev/null
+++ b/Engine/Plugins/VirtualProduction/Takes/Source/TakeRecorder/Private/Timecode/Visualization/Drawing/VisualizationDrawingUtils.h
@@ -0,0 +1,29 @@
+﻿// Copyright Epic Games, Inc. All Rights Reserved.
+
+#pragma once
+
+#include "EHitchDrawFlags.h"
+#include "HAL/Platform.h"
+
+struct FGeometry;
+class FSlateWindowElementList;
+
+namespace UE::TakeRecorder
+{
+struct FScrubRangeToScreen;
+struct FTimecodeHitchData;
+	
+/** Draws the top part (markers, etc.) of InData. */
+int32 DrawTimeSliderArea(
+	const FTimecodeHitchData& InData, const FScrubRangeToScreen& InRangeToScreen,
+	const FGeometry& InAllottedGeometry, FSlateWindowElementList& OutDrawElements, int32 InLayerId,
+	EHitchDrawFlags InFlags = EHitchDrawFlags::All
+	);
+	
+/** Draws the bottom part (markers, etc.) of InData. */
+int32 DrawTrackArea(
+	const FTimecodeHitchData& InData, const FScrubRangeToScreen& InRangeToScreen,
+	const FGeometry& InAllottedGeometry, FSlateWindowElementList& OutDrawElements, int32 InLayerId,
+	EHitchDrawFlags InFlags = EHitchDrawFlags::All
+	);
+}
diff --git a/Engine/Plugins/VirtualProduction/Takes/Source/TakeRecorder/Private/Timecode/Visualization/Sequencer/HitchVisualizationViewModel.cpp b/Engine/Plugins/VirtualProduction/Takes/Source/TakeRecorder/Private/Timecode/Visualization/Sequencer/HitchVisualizationViewModel.cpp
--- a/Engine/Plugins/VirtualProduction/Takes/Source/TakeRecorder/Private/Timecode/Visualization/Sequencer/HitchVisualizationViewModel.cpp
+++ b/Engine/Plugins/VirtualProduction/Takes/Source/TakeRecorder/Private/Timecode/Visualization/Sequencer/HitchVisualizationViewModel.cpp
@@ -2,8 +2,9 @@
 
 #include "HitchVisualizationViewModel.h"
 
+#include "Timecode/Visualization/Drawing/ScrubRangeToScreen.h"
+#include "Timecode/Visualization/Drawing/VisualizationDrawingUtils.h"
 #include "Timecode/Visualization/TakeRecorderHitchVisualizationSettings.h"
-#include "Timecode/Visualization/VisualizationDrawingUtils.h"
 
 namespace UE::TakeRecorder
 {
diff --git a/Engine/Plugins/VirtualProduction/Takes/Source/TakeRecorder/Private/Timecode/Visualization/Sequencer/HitchVisualizationViewModel.h b/Engine/Plugins/VirtualProduction/Takes/Source/TakeRecorder/Private/Timecode/Visualization/Sequencer/HitchVisualizationViewModel.h
--- a/Engine/Plugins/VirtualProduction/Takes/Source/TakeRecorder/Private/Timecode/Visualization/Sequencer/HitchVisualizationViewModel.h
+++ b/Engine/Plugins/VirtualProduction/Takes/Source/TakeRecorder/Private/Timecode/Visualization/Sequencer/HitchVisualizationViewModel.h
@@ -6,11 +6,13 @@
 #include "MVVM/Extensions/ITopTimeSliderOverlayExtension.h"
 #include "MVVM/Extensions/ITrackAreaOverlayExtension.h"
 #include "MVVM/ViewModels/ViewModel.h"
+#include "Misc/Attribute.h"
 #include "Timecode/Visualization/HitchAnalysis.h"
-#include "Timecode/Visualization/VisualizationDrawingUtils.h"
 
 namespace UE::TakeRecorder
 {
+struct FScrubRangeToScreen;
+	
 /**  */
 class FHitchVisualizationViewModel
 	: public Sequencer::FViewModel
diff --git a/Engine/Plugins/VirtualProduction/Takes/Source/TakeRecorder/Private/Timecode/Visualization/TakeRecorderHitchVisualizationSettings.cpp b/Engine/Plugins/VirtualProduction/Takes/Source/TakeRecorder/Private/Timecode/Visualization/TakeRecorderHitchVisualizationSettings.cpp
--- a/Engine/Plugins/VirtualProduction/Takes/Source/TakeRecorder/Private/Timecode/Visualization/TakeRecorderHitchVisualizationSettings.cpp
+++ b/Engine/Plugins/VirtualProduction/Takes/Source/TakeRecorder/Private/Timecode/Visualization/TakeRecorderHitchVisualizationSettings.cpp
@@ -2,7 +2,7 @@
 
 #include "TakeRecorderHitchVisualizationSettings.h"
 
-#include "VisualizationDrawingUtils.h"
+#include "Drawing/EHitchDrawFlags.h"
 
 bool UTakeRecorderHitchVisualizationSettings::GetShowFrameDropMarkers() const
 {
diff --git a/Engine/Plugins/VirtualProduction/Takes/Source/TakeRecorder/Private/Timecode/Visualization/VisualizationDrawingUtils.cpp b/Engine/Plugins/VirtualProduction/Takes/Source/TakeRecorder/Private/Timecode/Visualization/VisualizationDrawingUtils.cpp
deleted file mode 100644
diff --git a/Engine/Plugins/VirtualProduction/Takes/Source/TakeRecorder/Private/Timecode/Visualization/VisualizationDrawingUtils.h b/Engine/Plugins/VirtualProduction/Takes/Source/TakeRecorder/Private/Timecode/Visualization/VisualizationDrawingUtils.h
deleted file mode 100644
