Take Recorder: Fix non-unity error

diff --git a/Engine/Plugins/VirtualProduction/Takes/Source/TakeRecorder/Private/Timecode/Visualization/HitchViewOptionUtils.cpp b/Engine/Plugins/VirtualProduction/Takes/Source/TakeRecorder/Private/Timecode/Visualization/HitchViewOptionUtils.cpp
--- a/Engine/Plugins/VirtualProduction/Takes/Source/TakeRecorder/Private/Timecode/Visualization/HitchViewOptionUtils.cpp
+++ b/Engine/Plugins/VirtualProduction/Takes/Source/TakeRecorder/Private/Timecode/Visualization/HitchViewOptionUtils.cpp
@@ -5,13 +5,15 @@
 
 #include "Framework/MultiBox/MultiBoxBuilder.h"
 #include "Framework/MultiBox/MultiBoxExtender.h"
-#include "Timecode/Visualization/Sequencer/HitchVisualizationViewModel.h"
 #include "Hitching/FrameHitchSceneDecoration.h"
+#include "ISequencer.h"
 #include "MovieScene.h"
+#include "MovieSceneSequence.h"
 #include "ScopedTransaction.h"
 #include "TakeRecorderCommands.h"
 #include "TakeRecorderHitchVisualizationSettings.h"
 #include "Templates/SharedPointer.h"
+#include "Timecode/Visualization/Sequencer/HitchVisualizationViewModel.h"
 
 #define LOCTEXT_NAMESPACE "FHitchVisualizerInSequencer"
 
