Take Recorder: Add "Hitch Protection" section to Take Recorder.

When hitch protection is enabled:
- the timecode provider is overriden to be UTimecodeRegressionProvider,
- the engine timestep is overriden with the specified UCatchupFixedRateCustomTimeStep subclass,
- target and actual timecode are recorded into a UFrameHitchSceneDecoration, which is attached to the root ULevelSequence at the end of recording.

diff --git a/Engine/Plugins/VirtualProduction/Takes/Source/TakeRecorder/Private/Customization/TakeRecorderProjectSettingsCustomization.h b/Engine/Plugins/VirtualProduction/Takes/Source/TakeRecorder/Private/Customization/TakeRecorderProjectSettingsCustomization.h
--- a/Engine/Plugins/VirtualProduction/Takes/Source/TakeRecorder/Private/Customization/TakeRecorderProjectSettingsCustomization.h
+++ b/Engine/Plugins/VirtualProduction/Takes/Source/TakeRecorder/Private/Customization/TakeRecorderProjectSettingsCustomization.h
@@ -23,6 +23,8 @@ class FTakeRecorderProjectSettingsCustomization : public IDetailCustomization
 
 		// Pop the take recorder category to the top
 		DetailLayout.EditCategory("Take Recorder");
+		// Show the Hitch Protection category after "Take Recorder"
+		CustomizeHitchProtection(DetailLayout);
 
 		TArray<TWeakObjectPtr<UObject>> CustomizedObjects;
 		DetailLayout.GetObjectsBeingCustomized(CustomizedObjects);
@@ -67,4 +69,26 @@ class FTakeRecorderProjectSettingsCustomization : public IDetailCustomization
 			}
 		}
 	}
+
+private:
+
+	void CustomizeHitchProtection(IDetailLayoutBuilder& DetailLayout)
+	{
+		// Create a separate "Hitch Protection" category right under.
+		// ShowOnlyInnerProperties does not work correctly... we need to manually add the children ... otherwise they end up under the Take Recorder category.
+		IDetailCategoryBuilder& HitchProtectionCategory = DetailLayout.EditCategory("Hitch Protection");
+		const TSharedRef<IPropertyHandle> HitchProtectionSettings = DetailLayout.GetProperty(
+			GET_MEMBER_NAME_CHECKED(UTakeRecorderProjectSettings, HitchProtectionSettings)
+		);
+		
+		// Structs have a root node, under which their properties are nested. Since we're promoting them up, we need to hide this node.
+		HitchProtectionCategory.AddProperty(HitchProtectionSettings).EditCondition(false, {}).EditConditionHides(true);
+		
+		uint32 NumChildren = 0;
+		HitchProtectionSettings->GetNumChildren(NumChildren);
+		for (uint32 Index = 0; Index < NumChildren; ++Index)
+		{
+			HitchProtectionCategory.AddProperty(HitchProtectionSettings->GetChildHandle(Index));
+		}
+	}
 };
diff --git a/Engine/Plugins/VirtualProduction/Takes/Source/TakeRecorder/Private/Timecode/HitchlessProtectionRootLogic.cpp b/Engine/Plugins/VirtualProduction/Takes/Source/TakeRecorder/Private/Timecode/HitchlessProtectionRootLogic.cpp
--- a/Engine/Plugins/VirtualProduction/Takes/Source/TakeRecorder/Private/Timecode/HitchlessProtectionRootLogic.cpp
+++ b/Engine/Plugins/VirtualProduction/Takes/Source/TakeRecorder/Private/Timecode/HitchlessProtectionRootLogic.cpp
@@ -2,30 +2,13 @@
 
 #include "HitchlessProtectionRootLogic.h"
 
+#include "PropertyEditorModule.h"
+#include "Recorder/TakeRecorderHitchProtectionParameters.h"
 #include "Recorder/TakeRecorderSubsystem.h"
+#include "TakeRecorderSettings.h"
 
 namespace UE::TakeRecorder
 {
-namespace Private
-{
-static TAutoConsoleVariable<bool> CVarEnableHitchlessTimecodeProvider(
-	TEXT("TakeRecorder.EnableHitchlessTimecode"),
-	false,
-	TEXT("If true, we'll try to handle engine hitching. This requires that the UEngineCustomTimestep implements the IClockedTimeStep interface.")
-	);
-static TAutoConsoleVariable<float> CVarTimecodeLinearRegressionSeconds(
-	TEXT("TakeRecorder.TimecodeLinearRegressionSeconds"),
-	1.0f,
-	TEXT("When using linear regression to predict timecode, how many seconds worth of frames to sample.")
-	);
-
-static TAutoConsoleVariable<bool> CVarTestWithSource(
-	TEXT("Dominik.TestWithSource"),
-	true,
-	TEXT("When using linear regression to predict timecode, how many seconds worth of frames to sample.")
-);
-}
-	
 FHitchlessProtectionRootLogic::FHitchlessProtectionRootLogic()
 {
 	if (IsEngineStartupModuleLoadingComplete())
@@ -63,14 +46,13 @@ void FHitchlessProtectionRootLogic::SetupTakeRecorderDelegates()
 
 void FHitchlessProtectionRootLogic::OnInitializeRecording(UTakeRecorder* TakeRecorder)
 {
-	if (!Private::CVarEnableHitchlessTimecodeProvider.GetValueOnAnyThread())
+	const FTakeRecorderHitchProtectionParameters& Settings = GetMutableDefault<UTakeRecorderProjectSettings>()->HitchProtectionSettings;
+	if (!Settings.bEnableHitchProtection)
 	{
 		return;
 	}
 
-	const FTimecodeEstimationSetupResult TimecodeSetupResult = TimecodeRegressionRecording.SetupEngineEnvironmentForRecording(
-		Private::CVarTimecodeLinearRegressionSeconds.GetValueOnAnyThread()
-		);
+	const FTimecodeEstimationSetupResult TimecodeSetupResult = TimecodeRegressor.SetupEngineEnvironmentForRecording(Settings);
 	if (!TimecodeSetupResult)
 	{
 		return;
@@ -92,6 +74,6 @@ void FHitchlessProtectionRootLogic::OnCancelRecording(UTakeRecorder* TakeRecorde
 void FHitchlessProtectionRootLogic::CleanupRecording()
 {
 	RecordingState.Reset();
-	TimecodeRegressionRecording.CleanupRecording();
+	TimecodeRegressor.CleanupRecording();
 }
 }
diff --git a/Engine/Plugins/VirtualProduction/Takes/Source/TakeRecorder/Private/Timecode/HitchlessProtectionRootLogic.h b/Engine/Plugins/VirtualProduction/Takes/Source/TakeRecorder/Private/Timecode/HitchlessProtectionRootLogic.h
--- a/Engine/Plugins/VirtualProduction/Takes/Source/TakeRecorder/Private/Timecode/HitchlessProtectionRootLogic.h
+++ b/Engine/Plugins/VirtualProduction/Takes/Source/TakeRecorder/Private/Timecode/HitchlessProtectionRootLogic.h
@@ -30,7 +30,7 @@ private:
 	 * E.g. when a 1s hitch occurs at time 14:11:06:01, then the next frame should be 14:11:06:02 even though the real timecode is 14:11:07:02 meaning
 	 * that the engine has to catch up. This sets up the relevant systems.
 	 */
-	FTimecodeRegressionRecordSetup TimecodeRegressionRecording;
+	FTimecodeRegressionRecordSetup TimecodeRegressor;
 	
 	struct FRecordingState
 	{
diff --git a/Engine/Plugins/VirtualProduction/Takes/Source/TakeRecorder/Private/Timecode/Regression/TimecodeNotifications.cpp b/Engine/Plugins/VirtualProduction/Takes/Source/TakeRecorder/Private/Timecode/Regression/TimecodeNotifications.cpp
--- a/Engine/Plugins/VirtualProduction/Takes/Source/TakeRecorder/Private/Timecode/Regression/TimecodeNotifications.cpp
+++ b/Engine/Plugins/VirtualProduction/Takes/Source/TakeRecorder/Private/Timecode/Regression/TimecodeNotifications.cpp
@@ -22,202 +22,26 @@
 
 namespace UE::TakeRecorder
 {
-void ShowIncompatibleTimestepNotification(UEngineCustomTimeStep* InTimestep)
+void ShowNoTimecodeProviderNotification()
 {
 	if (FSlateApplication::IsInitialized()) // handle headless engine
 	{
-		FNotificationInfo Notification = LOCTEXT("IncompatibleTimestep.Title", "Hitchless timecode");
-		Notification.SubText = FText::Format(
-			LOCTEXT("IncompatibleTimestep.SubTextFmt", "Engine timestep {0} does not implement"),
-			InTimestep->GetClass()->GetDisplayNameText()
-			);
+		FNotificationInfo Notification = LOCTEXT("NoTimecodeProvider.Title", "No timecode provider set-up");
+		Notification.SubText = LOCTEXT("NoTimecodeProvider.SubTextFmt", "Hitch protection requires you to set up a timecode provider. Go to Project Settings.");
+		Notification.ExpireDuration = 5.f;
 		FSlateNotificationManager::Get().AddNotification(Notification)->SetCompletionState(SNotificationItem::CS_Fail);
 	}
 }
 
-namespace PickTimestepDialog
+void ShowNoTimeStepClassNotification()
 {
-class FIsClockedTimeStepFilter : public IClassViewerFilter
-{
-public:
-	virtual bool IsClassAllowed(
-		const FClassViewerInitializationOptions& InInitOptions,
-		const UClass* InClass, TSharedRef< class FClassViewerFilterFuncs > InFilterFuncs
-		) override
-	{
-		return InClass && InClass->IsChildOf<UEngineCustomTimeStep>() && InClass->ImplementsInterface(UClockedTimeStep::StaticClass());
-	}
-
-	virtual bool IsUnloadedClassAllowed(
-		const FClassViewerInitializationOptions& InInitOptions,
-		const TSharedRef<const IUnloadedBlueprintData> InUnloadedClassData,
-		TSharedRef<FClassViewerFilterFuncs> InFilterFuncs
-		) override
-	{
-		return false;
-	}
-};
-
-static TWeakObjectPtr<UClass>& GetTimestepToUseForEditorSession()
-{
-	static TWeakObjectPtr<UClass> TimestepToUseForEditorSession;
-	return TimestepToUseForEditorSession;
-}
-
-static TSharedRef<SWidget> MakeClassPicker(UClass*& PickedClass)
-{
-	FClassViewerInitializationOptions Options;
-	Options.InitiallySelectedClass = UCatchupFixedRateCustomTimeStep::StaticClass();
-	Options.Mode = EClassViewerMode::ClassPicker;
-	Options.NameTypeToDisplay = EClassViewerNameTypeToDisplay::DisplayName;
-	Options.bShowNoneOption = false;
-	Options.bIsActorsOnly = false;
-	Options.bShowUnloadedBlueprints = false;
-	Options.ClassFilters.Add(MakeShared<FIsClockedTimeStepFilter>());
-
-	FClassViewerModule& ClassViewerModule = FModuleManager::LoadModuleChecked<FClassViewerModule>("ClassViewer");
-	return ClassViewerModule.CreateClassViewer(
-		Options,
-		FOnClassPicked::CreateLambda([&PickedClass](UClass* InNewClass)
-		{
-			PickedClass = InNewClass;
-		}));
-}
-
-static TSharedRef<SWidget> MakeDialogContent(UClass*& PickedClass, bool& bRememberChoice)
-{
-	return SNew(SBox)
-	.Padding(5, 5)
-	.HAlign(HAlign_Fill)
-	[
-		SNew(SVerticalBox)
-
-		+SVerticalBox::Slot()
-		.AutoHeight()
-		.HAlign(HAlign_Left)
-		[
-			SNew(STextBlock)
-			.Text(LOCTEXT("FixedTimeStep.InstructionMessage",
-				"No engine timestep is set up in project settings.\nPlease select a timestep class to use:"
-				))
-		]
-
-		+SVerticalBox::Slot()
-		.Padding(4, 20)
-		.AutoHeight()
-		.HAlign(HAlign_Fill)
-		[
-			SNew(SVerticalBox)
-
-			// Class picker, which only shows timesteps implementing IClockedTimeStep
-			+SVerticalBox::Slot()
-			.AutoHeight()
-			[
-				SNew(SBorder)[ MakeClassPicker(PickedClass) ]
-			]
-			
-			// Remember checkbox
-			+SVerticalBox::Slot()
-			.AutoHeight()
-			.Padding(0, 4, 20, 0)
-			.HAlign(HAlign_Left)
-			[
-				SNew(SHorizontalBox)
-
-				+SHorizontalBox::Slot()
-				.AutoWidth()
-				[
-					SNew(SCheckBox)
-					.IsChecked_Lambda([&bRememberChoice]{ return bRememberChoice ? ECheckBoxState::Checked : ECheckBoxState::Unchecked; })
-					.OnCheckStateChanged_Lambda([&bRememberChoice](ECheckBoxState NewState){ bRememberChoice = NewState == ECheckBoxState::Checked; })
-				]
-				
-				+SHorizontalBox::Slot()
-				.AutoWidth()
-				.Padding(2, 0)
-				[
-					SNew(STextBlock)
-					.Text(LOCTEXT("FixedTimeStep.Remember", "Use this class for the rest of the editor session."))
-				]
-			]
-		]
-		
-		+SVerticalBox::Slot()
-		.AutoHeight()
-		.HAlign(HAlign_Left)
-		[
-			SNew(STextBlock)
-			.Text(LOCTEXT("FixedTimeStep.DisableMessage", "To disable this workflow, execute \"TakeRecorder.EnableHitchlessTimecode 0\"."))
-		]
-	];
-}
-	
-static UClass* AskUserToSelectCustomTimestep()
-{
-	if (!FSlateApplication::IsInitialized())
-	{
-		return nullptr;
-	}
-
-	if (UClass* ClassFromLastTime = GetTimestepToUseForEditorSession().Get())
-	{
-		return ClassFromLastTime;
-	}
-
-	bool bRememberChoice = true;
-	UClass* PickedClass = nullptr;
-	
-	// This message is only temporary / for testing until we come up with better UX.
-	const TSharedRef<SCustomDialog> Message = SNew(SCustomDialog)
-		.Title(LOCTEXT("FixedTimeStep.Title", "Fixed timestep required"))
-		.Content() [ PickTimestepDialog::MakeDialogContent(PickedClass, bRememberChoice) ]
-		.Buttons(
-		{
-			SCustomDialog::FButton(LOCTEXT("FixedTimeStep.Skip.Title", "Run without hitching correction"))
-				.SetToolTipText(LOCTEXT("FixedTimeStep.Skip.Tooltip",
-					"Skip hitchless timecode set-up. "
-					"If the engine hitches, timecode evaluation will jump by the length of the frame."))
-				.SetButtonRole(SCustomDialog::EButtonRole::Cancel),
-			SCustomDialog::FButton(LOCTEXT("FixedTimeStep.DisableCVar.Title", "Disable workflow"))
-				.SetToolTipText(LOCTEXT("FixedTimeStep.DisableCVar.Tooltip", "Executes TakeRecorder.EnableHitchlessTimecode 0 to disable hitchless setup"))
-				.SetOnClicked(FSimpleDelegate::CreateLambda([]
-				{
-					IConsoleVariable* CVar = IConsoleManager::Get().FindConsoleVariable(TEXT("TakeRecorder.EnableHitchlessTimecode"));
-					if (ensure(CVar))
-					{
-						CVar->Set(false, ECVF_SetByConsole);
-					}
-				})),
-			SCustomDialog::FButton(LOCTEXT("FixedTimeStep.UseDefault.Tooltip", "Use selected class"))
-				.SetToolTipText(LOCTEXT("FixedTimeStep.UseDefault.Title", "Uses the selected engine timestep."))
-				.SetPrimary(true)
-		});
-	
-	const int32 ButtonIndex = Message->ShowModal();
-	const bool bUseSelectedClass = ButtonIndex == 2;
-	if (bUseSelectedClass && bRememberChoice)
-	{
-		GetTimestepToUseForEditorSession() = PickedClass;
-	}
-	return bUseSelectedClass ? PickedClass : nullptr;
-}
-}
-	
-UEngineCustomTimeStep* AskUserToCreateFallbackTimestep()
-{
-	if (UClass* Class = PickTimestepDialog::AskUserToSelectCustomTimestep())
+	if (FSlateApplication::IsInitialized()) // handle headless engine
 	{
-		UEngineCustomTimeStep* CatchupTimestep = NewObject<UEngineCustomTimeStep>(GetTransientPackage(), Class);
-	
-		if (UCatchupFixedRateCustomTimeStep* Fixedtimestep = Cast<UCatchupFixedRateCustomTimeStep>(CatchupTimestep))
-		{
-			Fixedtimestep->FrameRate = GEngine->GetEngineSubsystem<UTakeRecorderSubsystem>()->GetFrameRate();
-		}
-
-		return CatchupTimestep;
+		FNotificationInfo Notification = LOCTEXT("NoTimeStepClass.Title", "Failed to create custom time step");
+		Notification.SubText = LOCTEXT("NoTimeStepClass.SubTextFmt", "Make sure that a valid custom timestep class is set for Hitch Protection");
+		Notification.ExpireDuration = 5.f;
+		FSlateNotificationManager::Get().AddNotification(Notification)->SetCompletionState(SNotificationItem::CS_Fail);
 	}
-	
-	return nullptr;
 }
 }
 
diff --git a/Engine/Plugins/VirtualProduction/Takes/Source/TakeRecorder/Private/Timecode/Regression/TimecodeNotifications.h b/Engine/Plugins/VirtualProduction/Takes/Source/TakeRecorder/Private/Timecode/Regression/TimecodeNotifications.h
--- a/Engine/Plugins/VirtualProduction/Takes/Source/TakeRecorder/Private/Timecode/Regression/TimecodeNotifications.h
+++ b/Engine/Plugins/VirtualProduction/Takes/Source/TakeRecorder/Private/Timecode/Regression/TimecodeNotifications.h
@@ -7,12 +7,9 @@ class UEngineCustomTimeStep;
 
 namespace UE::TakeRecorder
 {
-/** Shows a notification that the current timestep is incompatible. */
-void ShowIncompatibleTimestepNotification(UEngineCustomTimeStep* InTimestep);
-	
-/**
- * Shows a dialogue that ask the user to select a custom timestep.
- * @return The allocated timestep but not yet registered with the engine. 
- */
-UEngineCustomTimeStep* AskUserToCreateFallbackTimestep();
+/** Shows a notification that the user has not set up any timecode provider. */
+void ShowNoTimecodeProviderNotification();
+
+/** Shows a notification that the configured time step class could not be resolved. */
+void ShowNoTimeStepClassNotification();
 }
diff --git a/Engine/Plugins/VirtualProduction/Takes/Source/TakeRecorder/Private/Timecode/Regression/TimecodeRegressionRecordSetup.cpp b/Engine/Plugins/VirtualProduction/Takes/Source/TakeRecorder/Private/Timecode/Regression/TimecodeRegressionRecordSetup.cpp
--- a/Engine/Plugins/VirtualProduction/Takes/Source/TakeRecorder/Private/Timecode/Regression/TimecodeRegressionRecordSetup.cpp
+++ b/Engine/Plugins/VirtualProduction/Takes/Source/TakeRecorder/Private/Timecode/Regression/TimecodeRegressionRecordSetup.cpp
@@ -5,10 +5,8 @@
 #include "CatchupFixedRateCustomTimeStep.h"
 #include "ClassViewerFilter.h"
 #include "Engine/Engine.h"
-#include "Engine/EngineCustomTimeStep.h"
 #include "Estimation/TimecodeRegressionProvider.h"
-#include "Framework/Notifications/NotificationManager.h"
-#include "Misc/CoreDelegates.h"
+#include "Recorder/TakeRecorderHitchProtectionParameters.h"
 #include "Recorder/TakeRecorderSubsystem.h"
 #include "TimecodeNotifications.h"
 #include "UObject/Package.h"
@@ -18,50 +16,51 @@ namespace UE::TakeRecorder
 namespace Private
 {
 /** @return The timecode estimator instance we'll use to estimate timecode - effectively handling any engine hitches. */
-static UTimecodeRegressionProvider* CreateTimecodeEstimator(UTimecodeProvider* InActualImplementation, float InRegressionFramesInSeconds)
+static UTimecodeRegressionProvider* CreateTimecodeEstimator(
+	UTimecodeProvider* InActualImplementation, const FFrameRate& InRecordFrameRate, float InRegressionFramesInSeconds
+	)
 {
 	UTimecodeRegressionProvider* TimecodeRegression = NewObject<UTimecodeRegressionProvider>(GetTransientPackage());
 	TimecodeRegression->SetTimecodeImplementation(InActualImplementation);
 	
 	// How many frames do we need to sample?
 	// E.g. if FrameRate is 60 FPS of 60 frames, then Numerator = 60 meaning 60 / 1 = 60. We want
-	const FFrameRate FrameRate = GEngine->GetEngineSubsystem<UTakeRecorderSubsystem>()->GetFrameRate();
-	TimecodeRegression->NumSampledFrames = FrameRate.Denominator == 0
-		? 1 : InRegressionFramesInSeconds * FrameRate.Numerator / FrameRate.Denominator;
+	TimecodeRegression->NumSampledFrames = InRecordFrameRate.Denominator == 0
+		? 1 : InRegressionFramesInSeconds * InRecordFrameRate.Numerator / InRecordFrameRate.Denominator;
 	
 	return TimecodeRegression;
 }
 }
 
-FTimecodeEstimationSetupResult FTimecodeRegressionRecordSetup::SetupEngineEnvironmentForRecording(float InRegressionFramesInSeconds)
+FTimecodeEstimationSetupResult FTimecodeRegressionRecordSetup::SetupEngineEnvironmentForRecording(
+	const FTakeRecorderHitchProtectionParameters& InParams
+	)
 {
 	FTimecodeEstimationSetupResult Result;
 	
 	UTimecodeProvider* CurrentTimecodeProvider = GEngine ? GEngine->GetTimecodeProvider() : nullptr;
 	if (!CurrentTimecodeProvider)
 	{
-		return Result;
-	}
-
-	UEngineCustomTimeStep* PreexistingTimestep = GEngine->GetCustomTimeStep();
-	if (PreexistingTimestep && !UTimecodeRegressionProvider::SupportsTimestep(PreexistingTimestep))
-	{
-		ShowIncompatibleTimestepNotification(PreexistingTimestep);
-		return Result;
+		ShowNoTimecodeProviderNotification();
+		return {};
 	}
+	const FFrameRate FrameRate = GEngine->GetEngineSubsystem<UTakeRecorderSubsystem>()->GetFrameRate();
 
-	UTimecodeRegressionProvider* Estimator = Private::CreateTimecodeEstimator(CurrentTimecodeProvider, InRegressionFramesInSeconds);
-	if (PreexistingTimestep)
+	UClass* TimestepClass = InParams.CustomTimestep.LoadSynchronous();
+	UCatchupFixedRateCustomTimeStep* Timestep = TimestepClass ? NewObject<UCatchupFixedRateCustomTimeStep>(GetTransientPackage(), TimestepClass) : nullptr;
+	if (!Timestep)
 	{
-		// Use pre-configured custom timestep
-		RecordData.Emplace(Estimator);
+		ShowNoTimeStepClassNotification();
+		return {};
 	}
-	else if (UEngineCustomTimeStep* DefaultTimestep = AskUserToCreateFallbackTimestep())
-	{
-		// There was no pre-configured custom timestep... clean it up when recording finishes.
-		RecordData.Emplace(DefaultTimestep, Estimator);
-	}
-
+	Timestep->FrameRate = FrameRate;
+	Timestep->MaxCatchupSeconds = InParams.MaxCatchupSeconds;
+	
+	UTimecodeRegressionProvider* Estimator = Private::CreateTimecodeEstimator(
+		CurrentTimecodeProvider, FrameRate, InParams.RegressionBufferSizeInSeconds
+		);
+	RecordData.Emplace(Timestep, Estimator);
+	
 	Result.Estimator = Estimator;
 	return Result;
 }
diff --git a/Engine/Plugins/VirtualProduction/Takes/Source/TakeRecorder/Private/Timecode/Regression/TimecodeRegressionRecordSetup.h b/Engine/Plugins/VirtualProduction/Takes/Source/TakeRecorder/Private/Timecode/Regression/TimecodeRegressionRecordSetup.h
--- a/Engine/Plugins/VirtualProduction/Takes/Source/TakeRecorder/Private/Timecode/Regression/TimecodeRegressionRecordSetup.h
+++ b/Engine/Plugins/VirtualProduction/Takes/Source/TakeRecorder/Private/Timecode/Regression/TimecodeRegressionRecordSetup.h
@@ -9,6 +9,8 @@
 #include "Recorder/TakeRecorder.h"
 #include "Templates/UnrealTemplate.h"
 
+struct FTakeRecorderHitchProtectionParameters;
+
 namespace UE::TakeRecorder
 {
 /** Holds info about when the engine environment has been set up to estimate timecode. */
@@ -31,11 +33,10 @@ public:
 
 	/**
 	 * Sets up timecode and optionally the custom time step.
-	 * @param InRegressionFramesInSeconds Determines size of the linear regression buffer. This is the number of seconds worth of frames we should use.
-	 *	Example: 1s worth of frames means a buffer size of 24 timecode samples if the UTakeRecorder is set to record at 24 FPS.
+	 * @param InParams Settings for setting up recording.
 	 * @return The real timecode provider that we will be estimating. Unset if initialization failed.
 	 */
-	FTimecodeEstimationSetupResult SetupEngineEnvironmentForRecording(float InRegressionFramesInSeconds);
+	FTimecodeEstimationSetupResult SetupEngineEnvironmentForRecording(const FTakeRecorderHitchProtectionParameters& InParams);
 
 	/** Clears the time step and timecode provider overrides. */
 	void CleanupRecording() { RecordData.Reset(); }
diff --git a/Engine/Plugins/VirtualProduction/Takes/Source/TakeRecorder/Public/Recorder/TakeRecorderHitchProtectionParameters.h b/Engine/Plugins/VirtualProduction/Takes/Source/TakeRecorder/Public/Recorder/TakeRecorderHitchProtectionParameters.h
new file mode 100644
--- /dev/null
+++ b/Engine/Plugins/VirtualProduction/Takes/Source/TakeRecorder/Public/Recorder/TakeRecorderHitchProtectionParameters.h
@@ -0,0 +1,50 @@
+﻿// Copyright Epic Games, Inc. All Rights Reserved.
+
+#pragma once
+
+#include "CatchupFixedRateCustomTimeStep.h"
+#include "Engine/EngineCustomTimeStep.h"
+#include "TakeRecorderHitchProtectionParameters.generated.h"
+
+USTRUCT()
+struct FTakeRecorderHitchProtectionParameters
+{
+	GENERATED_BODY()
+
+	/**
+	 * Attempt to prevent hitches in recorded data. If a frame were to be dropped, Take Recorder will instead slow its evaluation and attempt to
+	 * catch up. This is achieved by using a fixed delta time time-step, which will use the FPS value set in Take Recorder.
+	 */
+	UPROPERTY(Config, EditAnywhere, Category = "Take Recorder")
+	bool bEnableHitchProtection = true;
+
+	/**
+	 * Determines the number of frames in the buffer used by linear regression.
+	 * The final number of frames is FPS * RegressionBufferSizeInSeconds.
+	 * 
+	 * Without hitch protection, each frame is marked with the timecode reported by UTimecodeProvider; when a hitch occurs, the timecode hitches by
+	 * the same amount.
+	 * 
+	 * With hitch protection is enabled, the engine is set up with a fixed delta time step. Every frame, the timecode is sampled from the timecode
+	 * provider, stamped with the current physical time, and placed in the buffer. When frame n takes longer to evaluate than expected ("hitches"), the
+	 * phyiscal time will have moved by a certain amount but the world's simulation time will have remained the same (as simulation time advances only
+	 * when the engine ticks). Then, frame n+1 uses the world's simulation time to perform linear regression on what the frame's timecode should be.
+	 */
+	UPROPERTY(Config, AdvancedDisplay, EditAnywhere, Category = "Take Recorder", meta = (EditCondition = "bEnableHitchProtection", EditConditionHides))
+	float RegressionBufferSizeInSeconds = 1.f;
+
+	/** The custom engine timestep class to use. Allows you to customize the behavior. */
+	UPROPERTY(Config, EditAnywhere, Category = "Take Recorder", meta=(EditCondition = "bEnableHitchProtection", EditConditionHides, AllowAbstract = "false"))
+	TSoftClassPtr<UCatchupFixedRateCustomTimeStep> CustomTimestep = UCatchupFixedRateCustomTimeStep::StaticClass();
+
+	/**
+	 * Maximum allowed delay when Take Recorder is attempting to catch up during protection.
+	 * 
+	 * If Take Recorder falls behind further than this value, it will not accrue additional delay and will drop frames. The result is that the next
+	 * frame immediately jumps to the target timecode in order to attempt to catch up.
+	 *
+	 * Larger values will allow for jumps in timecode in the final data but will also result in possible larger delays in evaluation.
+	 */
+	UPROPERTY(Config, EditAnywhere, Category = "Take Recorder", meta = (EditCondition = "bEnableHitchProtection", EditConditionHides))
+	double MaxCatchupSeconds = 8.f;
+};
\ No newline at end of file
diff --git a/Engine/Plugins/VirtualProduction/Takes/Source/TakeRecorder/Public/TakeRecorderSettings.h b/Engine/Plugins/VirtualProduction/Takes/Source/TakeRecorder/Public/TakeRecorderSettings.h
--- a/Engine/Plugins/VirtualProduction/Takes/Source/TakeRecorder/Public/TakeRecorderSettings.h
+++ b/Engine/Plugins/VirtualProduction/Takes/Source/TakeRecorder/Public/TakeRecorderSettings.h
@@ -3,6 +3,7 @@
 #pragma once
 
 #include "UObject/Object.h"
+#include "Recorder/TakeRecorderHitchProtectionParameters.h"
 #include "Recorder/TakeRecorderParameters.h"
 #include "TakeRecorderSettings.generated.h"
 
@@ -65,9 +66,25 @@ public:
 	virtual void PostEditChangeProperty(FPropertyChangedEvent& PropertyChangedEvent) override;
 #endif
 
+	/** General take recorder settings */
 	UPROPERTY(config, EditAnywhere, BlueprintReadOnly, Category="Take Recorder", meta=(ShowOnlyInnerProperties))
 	FTakeRecorderProjectParameters Settings;
 
+	/**
+	 * Controls the behaviour for hitch protection.
+	 *
+	 * On a high level, it achieves even key spacing in recorded data even when the engine freezes on a frame ("hitches").
+	 * When a frames takes longer, the simulation continues to advance at the same fixed time step and assigned to the timecode the frame would have
+	 * had, if they hitch hadn't occured.
+	 * 
+	 * Example: frame n at 17:56:12:02 freezes for 1s. In physical time, once the hitches finishes, timecode will have advanced by 1s to 17:56:13:02.
+	 * However, hitch protection annotates frame n+1 with timecode 17:56:12:03. The subsequent frames continue to be ticked at the same rate, e.g. 1/24s;
+	 * as long as it takes less CPU time to run the simulation (i.e. CPU time to simulate is less than 1/24s), the simulation time eventually catches
+	 * up with physical time.
+	 */
+	UPROPERTY(config, EditAnywhere, Category="Hitch Protection")
+	FTakeRecorderHitchProtectionParameters HitchProtectionSettings;
+
 	/** Array of externally supplied CDOs that should be displayed on the take recorder project settings */
 	TArray<TWeakObjectPtr<UObject>> AdditionalSettings;
 };
