Take Recorder: Clamp FTakeRecorderHitchProtectionParameters::RegressionBufferSizeInSeconds.

diff --git a/Engine/Plugins/VirtualProduction/Takes/Source/TakeRecorder/Private/Timecode/Regression/TimecodeRegressionRecordSetup.cpp b/Engine/Plugins/VirtualProduction/Takes/Source/TakeRecorder/Private/Timecode/Regression/TimecodeRegressionRecordSetup.cpp
--- a/Engine/Plugins/VirtualProduction/Takes/Source/TakeRecorder/Private/Timecode/Regression/TimecodeRegressionRecordSetup.cpp
+++ b/Engine/Plugins/VirtualProduction/Takes/Source/TakeRecorder/Private/Timecode/Regression/TimecodeRegressionRecordSetup.cpp
@@ -25,6 +25,7 @@ static UTimecodeRegressionProvider* CreateTimecodeEstimator(
 	
 	// How many frames do we need to sample?
 	// E.g. if FrameRate is 60 FPS of 60 frames, then Numerator = 60 meaning 60 / 1 = 60. We want
+	InRegressionFramesInSeconds = FMath::Clamp(InRegressionFramesInSeconds, 0.0166, 60);
 	TimecodeRegression->NumSampledFrames = InRecordFrameRate.Denominator == 0
 		? 1 : InRegressionFramesInSeconds * InRecordFrameRate.Numerator / InRecordFrameRate.Denominator;
 	
diff --git a/Engine/Plugins/VirtualProduction/Takes/Source/TakeRecorder/Public/Recorder/TakeRecorderHitchProtectionParameters.h b/Engine/Plugins/VirtualProduction/Takes/Source/TakeRecorder/Public/Recorder/TakeRecorderHitchProtectionParameters.h
--- a/Engine/Plugins/VirtualProduction/Takes/Source/TakeRecorder/Public/Recorder/TakeRecorderHitchProtectionParameters.h
+++ b/Engine/Plugins/VirtualProduction/Takes/Source/TakeRecorder/Public/Recorder/TakeRecorderHitchProtectionParameters.h
@@ -30,7 +30,7 @@ struct FTakeRecorderHitchProtectionParameters
 	 * phyiscal time will have moved by a certain amount but the world's simulation time will have remained the same (as simulation time advances only
 	 * when the engine ticks). Then, frame n+1 uses the world's simulation time to perform linear regression on what the frame's timecode should be.
 	 */
-	UPROPERTY(Config, AdvancedDisplay, EditAnywhere, Category = "Take Recorder", meta = (EditCondition = "bEnableHitchProtection", EditConditionHides))
+	UPROPERTY(Config, AdvancedDisplay, EditAnywhere, Category = "Take Recorder", meta = (EditCondition = "bEnableHitchProtection", EditConditionHides, ClampMin = 0.01666, ClampMax = 10))
 	float RegressionBufferSizeInSeconds = 1.f;
 
 	/** The custom engine timestep class to use. Allows you to customize the behavior. */
