Take Recorder: Hitch visualization UI now detects hovered items:
- Hovered elements are drawn brighter
- Shows a tooltip on hovered items

diff --git a/Engine/Plugins/VirtualProduction/Takes/Source/TakeRecorder/Private/TakeRecorderStyle.cpp b/Engine/Plugins/VirtualProduction/Takes/Source/TakeRecorder/Private/TakeRecorderStyle.cpp
--- a/Engine/Plugins/VirtualProduction/Takes/Source/TakeRecorder/Private/TakeRecorderStyle.cpp
+++ b/Engine/Plugins/VirtualProduction/Takes/Source/TakeRecorder/Private/TakeRecorderStyle.cpp
@@ -211,12 +211,22 @@ FTakeRecorderStyle::FTakeRecorderStyle()
 
 	// Hitch visualization
 	const FLinearColor ErrorColor(0.36f, 0.04f, 0.04f);
-	const FLinearColor WarningColor(0.31f, 0.31f, 0.03f);
-	Set("Hitching.ErrorColor", ErrorColor); 
-	Set("Hitching.WarningColor", WarningColor); 
-	Set("Hitching.CatchUpZoneBrush", new FSlateColorBrush(FLinearColor(WarningColor.R, WarningColor.G, WarningColor.B, 0.5f))); // time slider area where we were behind target timecode
-	Set("Hitching.MarkerWarningBrush", new IMAGE_BRUSH_SVG("AlertTriangleSolid", Icon16x16));
-
+	const FLinearColor ErrorColor_Hovered(1.00f, 0.05f, 0.05f);
+	const FLinearColor WarningColor(0.58f, 0.31f, 0.00f, 0.90f);
+	const FLinearColor WarningColor_Hovered(1.00f, 0.38f, 0.00f, 0.90f);
+	
+	Set("Hitching.SkipMarker.Normal", ErrorColor); 
+	Set("Hitching.SkipMarker.Hovered", ErrorColor_Hovered); 
+	Set("Hitching.SkipMarker.Icon", new IMAGE_BRUSH_SVG("AlertTriangleSolid", Icon16x16));
+	
+	Set("Hitching.RepeatMarker.Normal", ErrorColor); 
+	Set("Hitching.RepeatMarker.Hovered", ErrorColor_Hovered); 
+	Set("Hitching.RepeatMarker.Icon", new IMAGE_BRUSH_SVG("AlertTriangleSolid", Icon16x16));
+	
+	Set("Hitching.CatchupZone.Normal", WarningColor); 
+	Set("Hitching.CatchupZone.Hovered", WarningColor_Hovered); 
+	Set("Hitching.CatchupZone.Icon", new FSlateColorBrush(FLinearColor::White)); // The colour does not matter - we'll tint it in the drawing code.
+	
 	FSlateStyleRegistry::RegisterSlateStyle(*this);
 }
 
diff --git a/Engine/Plugins/VirtualProduction/Takes/Source/TakeRecorder/Private/Timecode/Visualization/Drawing/HitchHoverInfo.h b/Engine/Plugins/VirtualProduction/Takes/Source/TakeRecorder/Private/Timecode/Visualization/Drawing/HitchHoverInfo.h
--- a/Engine/Plugins/VirtualProduction/Takes/Source/TakeRecorder/Private/Timecode/Visualization/Drawing/HitchHoverInfo.h
+++ b/Engine/Plugins/VirtualProduction/Takes/Source/TakeRecorder/Private/Timecode/Visualization/Drawing/HitchHoverInfo.h
@@ -2,10 +2,49 @@
 
 #pragma once
 
+#include "Misc/Optional.h"
+#include "HAL/Platform.h"
+
 namespace UE::TakeRecorder
 {
+/**
+ * Information about which hitch visualization geometry is hovered.
+ * Only one of the optionals is set at any time.
+ */
 struct FHitchHoverInfo
 {
+	/** Index to the FTimecodeHitchData::SkippedTimecodeMarkers that should be drawn as set. */
+	TOptional<int32> SkippedMarkerIndex;
 	
+	/** Index to the FTimecodeHitchData::RepeatedTimecodeMarkers that should be drawn as set. */
+	TOptional<int32> RepeatedMarkerIndex;
+
+	/** Index to the FTimecodeHitchData::RepeatedTimecodeMarkers that should be drawn as hovered. */
+	TOptional<int32> CatchupTimeIndex;
+
+	/** Clears the hover state. */
+	void Reset()
+	{
+		SkippedMarkerIndex.Reset();
+		RepeatedMarkerIndex.Reset();
+		CatchupTimeIndex.Reset();
+	}
+
+	/** @return Whether anything is hovered */
+	bool IsHovered() const { return SkippedMarkerIndex || RepeatedMarkerIndex || CatchupTimeIndex; }
+	/** @return Whether anything is hovered */
+	operator bool() const { return IsHovered(); }
+
+	friend bool operator==(const FHitchHoverInfo& Left, const FHitchHoverInfo& Right)
+	{
+		return Left.SkippedMarkerIndex == Right.SkippedMarkerIndex
+			&& Left.RepeatedMarkerIndex == Right.RepeatedMarkerIndex
+			&& Left.CatchupTimeIndex == Right.CatchupTimeIndex;
+	}
+
+	friend bool operator!=(const FHitchHoverInfo& Left, const FHitchHoverInfo& Right)
+	{
+		return !(Left == Right);
+	}
 };
 }
diff --git a/Engine/Plugins/VirtualProduction/Takes/Source/TakeRecorder/Private/Timecode/Visualization/Drawing/HitchVisualizationStyle.h b/Engine/Plugins/VirtualProduction/Takes/Source/TakeRecorder/Private/Timecode/Visualization/Drawing/HitchVisualizationStyle.h
new file mode 100644
--- /dev/null
+++ b/Engine/Plugins/VirtualProduction/Takes/Source/TakeRecorder/Private/Timecode/Visualization/Drawing/HitchVisualizationStyle.h
@@ -0,0 +1,53 @@
+﻿// Copyright Epic Games, Inc. All Rights Reserved.
+
+#pragma once
+
+#include "TakeRecorderStyle.h"
+#include "Styling/SlateBrush.h"
+
+namespace UE::TakeRecorder
+{
+struct FMarkerStyle
+{
+	/** Color that the marker vertical line & warning icon should have when drawn normally. */
+	FLinearColor NormalTint = FTakeRecorderStyle::Get().GetColor("Hitching.SkipMarker.Normal");
+	
+	/** Color that the marker vertical line & warning icon should have when drawn hovered. */
+	FLinearColor HoveredTint = FTakeRecorderStyle::Get().GetColor("Hitching.SkipMarker.Hovered");
+	
+	/** Brush to drawn next to vertical line on marker. */
+	const FSlateBrush* IconBrush = FTakeRecorderStyle::Get().GetBrush("Hitching.SkipMarker.Icon");
+
+	/** Offset that marker icon should have from its vertical line. */
+	FVector2f MarkerBrushOffset{ 4.f, 4.f };
+};
+
+struct FCatchupZoneStyle
+{
+	/** Color that the catchup zone vertical lines & box brush should have when drawn normally. */
+	FLinearColor CatchupZone_Normal = FTakeRecorderStyle::Get().GetColor("Hitching.CatchupZone.Normal");
+	
+	/** Color that the catchup zone vertical lines & box brush should have when drawn hovered. */
+	FLinearColor CatchupZone_Hovered = FTakeRecorderStyle::Get().GetColor("Hitching.CatchupZone.Hovered");
+	
+	/** Brush to drawn between left and right vertical line in catchup zone. */
+	const FSlateBrush* CatchupZoneBrush = FTakeRecorderStyle::Get().GetBrush("Hitching.CatchupZone.Icon");
+};
+	
+struct FHitchVisualizationStyle
+{
+	/** How to draw markers for when the engine skipped timecode on a frame. */
+	FMarkerStyle SkipMarkerStyle;
+	
+	/** How to draw markers for when the engine repeated timecode. */
+	FMarkerStyle RepeatMarkerStyle
+	{
+		FTakeRecorderStyle::Get().GetColor("Hitching.SkipMarker.Normal"),
+		FTakeRecorderStyle::Get().GetColor("Hitching.SkipMarker.Hovered"),
+		FTakeRecorderStyle::Get().GetBrush("Hitching.SkipMarker.Icon")
+	};
+	
+	/** How to draw the catchup zone, i.e. an area in which the engine was trailing target timecode. */
+	FCatchupZoneStyle CatchupZoneStyle;
+};
+}
diff --git a/Engine/Plugins/VirtualProduction/Takes/Source/TakeRecorder/Private/Timecode/Visualization/Drawing/VisualizationDrawingUtils.cpp b/Engine/Plugins/VirtualProduction/Takes/Source/TakeRecorder/Private/Timecode/Visualization/Drawing/VisualizationDrawingUtils.cpp
--- a/Engine/Plugins/VirtualProduction/Takes/Source/TakeRecorder/Private/Timecode/Visualization/Drawing/VisualizationDrawingUtils.cpp
+++ b/Engine/Plugins/VirtualProduction/Takes/Source/TakeRecorder/Private/Timecode/Visualization/Drawing/VisualizationDrawingUtils.cpp
@@ -2,24 +2,28 @@
 
 #include "VisualizationDrawingUtils.h"
 
+#include "HitchHoverInfo.h"
+#include "HitchVisualizationStyle.h"
 #include "Misc/FrameTime.h"
 #include "Rendering/DrawElementTypes.h"
 #include "ScrubRangeToScreen.h"
-#include "TakeRecorderStyle.h"
 #include "Timecode/Visualization/HitchAnalysis.h"
 
 namespace UE::TakeRecorder
 {
 namespace Private
 {
-const FVector2f MarkerBrushOffset{ 4.f, 4.f };
-	
+/** Packs args that are used together for drawing geometry in Sequencer's widget overlay areas. */
 struct FDrawArgs
 {
+	/** Converts FFrameTime to Sequencer's widget area. */
 	const FScrubRangeToScreen& RangeToScreen;
+	/** The geometry area allocated to the overlay */
 	const FGeometry& AllottedGeometry;
+	/** The base layer ID to draw onto */
 	int32 LayerId;
-	
+
+	/** The list to append draw elements to */
 	FSlateWindowElementList& OutDrawElements;
 
 	explicit FDrawArgs(const FScrubRangeToScreen& RangeToScreen, const FGeometry& AllottedGeometry, int32 LayerId, FSlateWindowElementList& OutDrawElements)
@@ -31,13 +35,12 @@ struct FDrawArgs
 };
 
 static FGeometry MakeMarkerIconGeometry(
-	const FFrameTime& InTime, const FScrubRangeToScreen& InRangeToScreen, const FGeometry& InAllottedGeometry
+	const FFrameTime& InTime, const FScrubRangeToScreen& InRangeToScreen, const FGeometry& InAllottedGeometry, const FMarkerStyle& InStyle
 	)
 {
 	const float MarkerPos = InRangeToScreen.FrameToLocalX(InTime);
-	const FVector2f Position = FVector2f(MarkerPos, 0) + MarkerBrushOffset;
-	const FSlateBrush* Brush = FTakeRecorderStyle::Get().GetBrush("Hitching.MarkerWarningBrush");
-	const FVector2f Size = Brush->GetImageSize();
+	const FVector2f Position = FVector2f(MarkerPos, 0) + InStyle.MarkerBrushOffset;
+	const FVector2f Size = InStyle.IconBrush->GetImageSize();
 	return InAllottedGeometry.MakeChild(Size, FSlateLayoutTransform(Position));
 }
 
@@ -64,8 +67,7 @@ static FGeometry MakeCatchupZoneGeometry(
 }
 	
 static void DrawVerticalLine(
-	const FFrameTime& InFrameTime, const FDrawArgs& InDrawArgs,
-	const FLinearColor& InColor
+	const FFrameTime& InFrameTime, const FDrawArgs& InDrawArgs, const FLinearColor& InColor
 	)
 {
 	const float LinePos = InDrawArgs.RangeToScreen.FrameToLocalX(InFrameTime);
@@ -94,111 +96,270 @@ static void DrawVerticalLine(
 	);
 }
 	
-static void DrawVerticalLines(TConstArrayView<FSkippedTimecodeMarker> InMarkers, const FDrawArgs& InDrawArgs)
+static void DrawMarkerVerticalLines(
+	TConstArrayView<FUnexpectedTimecodeMarker> InMarkers, const TOptional<int32>& InHoveredIndex, const FMarkerStyle& InStyle,
+	const FDrawArgs& InDrawArgs
+	)
 {
-	const FLinearColor& Color = FTakeRecorderStyle::Get().GetColor("Hitching.ErrorColor");
-	for (const FSkippedTimecodeMarker& Marker : InMarkers)
+	for (int32 Index = 0; Index < InMarkers.Num(); ++Index)
 	{
-		DrawVerticalLine(Marker.Frame, InDrawArgs, Color);
+		const FUnexpectedTimecodeMarker& Marker = InMarkers[Index];
+		const bool bIsHovered = InHoveredIndex && *InHoveredIndex == Index;
+		DrawVerticalLine(Marker.Frame, InDrawArgs, bIsHovered ? InStyle.HoveredTint : InStyle.NormalTint);
 	}
 }
 
-static void DrawVerticalLines(TConstArrayView<FCatchupTimeRange> InTimes, const FDrawArgs& InDrawArgs)
+static void DrawCatchupVerticalLines(
+	TConstArrayView<FCatchupTimeRange> InTimes, const TOptional<int32>& InHoveredIndex, const FCatchupZoneStyle& InStyle,
+	const FDrawArgs& InDrawArgs
+	)
 {
-	const FLinearColor& Color = FTakeRecorderStyle::Get().GetColor("Hitching.WarningColor");
-	for (const FCatchupTimeRange& CatchupRange : InTimes)
+	for (int32 Index = 0; Index < InTimes.Num(); ++Index)
 	{
-		DrawVerticalLine(CatchupRange.StartTime, InDrawArgs, Color);
-		DrawVerticalLine(GetCatchupRangeLastVerticalLineFrame(CatchupRange), InDrawArgs, Color);
+		const FCatchupTimeRange& CatchupRange = InTimes[Index];
+		const bool bIsHovered = InHoveredIndex && *InHoveredIndex == Index;
+		const FLinearColor& InColor = bIsHovered ? InStyle.CatchupZone_Hovered : InStyle.CatchupZone_Normal;
+		
+		DrawVerticalLine(CatchupRange.StartTime, InDrawArgs, InColor);
+		DrawVerticalLine(GetCatchupRangeLastVerticalLineFrame(CatchupRange), InDrawArgs, InColor);
 	}
 }
 	
-static void DrawMarkerWarningIcons(TConstArrayView<FSkippedTimecodeMarker> InMarkers, const FDrawArgs& InDrawArgs)
+static void DrawMarkerWarningIcons(
+	TConstArrayView<FUnexpectedTimecodeMarker> InMarkers, const TOptional<int32>& InHoveredIndex, const FMarkerStyle& InStyle,
+	const FDrawArgs& InDrawArgs
+	)
 {
-	for (const FSkippedTimecodeMarker& Marker : InMarkers)
+	for (int32 Index = 0; Index < InMarkers.Num(); ++Index)
 	{
-		const FSlateBrush* Brush = FTakeRecorderStyle::Get().GetBrush("Hitching.MarkerWarningBrush");
+		const FUnexpectedTimecodeMarker& Marker = InMarkers[Index];
+		const bool bIsHovered = InHoveredIndex && *InHoveredIndex == Index;
 		FSlateDrawElement::MakeBox(
 			InDrawArgs.OutDrawElements,
 			InDrawArgs.LayerId,
-			MakeMarkerIconGeometry(Marker.Frame, InDrawArgs.RangeToScreen, InDrawArgs.AllottedGeometry).ToPaintGeometry(),
-			Brush,
+			MakeMarkerIconGeometry(Marker.Frame, InDrawArgs.RangeToScreen, InDrawArgs.AllottedGeometry, InStyle).ToPaintGeometry(),
+			InStyle.IconBrush,
 			ESlateDrawEffect::None,
-			FTakeRecorderStyle::Get().GetColor("Hitching.ErrorColor")
+			bIsHovered ? InStyle.HoveredTint : InStyle.NormalTint
 			);
 	}
 }
 
-static void DrawCatchupZones(TConstArrayView<FCatchupTimeRange> InTimes, const FDrawArgs& InDrawArgs)
+static void DrawCatchupZones(
+	TConstArrayView<FCatchupTimeRange> InTimes,  const TOptional<int32>& InHoveredIndex, const FCatchupZoneStyle& InStyle,
+	const FDrawArgs& InDrawArgs
+	)
 {
-	for (const FCatchupTimeRange& CatchupRange : InTimes)
+	for (int32 Index = 0; Index < InTimes.Num(); ++Index)
 	{
+		const FCatchupTimeRange& CatchupRange = InTimes[Index];
+		const bool bIsHovered = InHoveredIndex && *InHoveredIndex == Index;
 		FSlateDrawElement::MakeBox(
 			InDrawArgs.OutDrawElements,
 			InDrawArgs.LayerId,
 			MakeCatchupZoneGeometry(CatchupRange, InDrawArgs.RangeToScreen, InDrawArgs.AllottedGeometry).ToPaintGeometry(),
-			FTakeRecorderStyle::Get().GetBrush("Hitching.CatchUpZoneBrush"),
+			InStyle.CatchupZoneBrush,
 			ESlateDrawEffect::None,
-			FTakeRecorderStyle::Get().GetColor("Hitching.WarningColor")
+			bIsHovered ? InStyle.CatchupZone_Hovered : InStyle.CatchupZone_Normal
 			);
 	}
 }
 }
 	
 int32 DrawTimeSliderArea(
-	const FTimecodeHitchData& InData, const FScrubRangeToScreen& InRangeToScreen,
+	const FTimecodeHitchData& InData, const FHitchHoverInfo& InHoverInfo, const FScrubRangeToScreen& InRangeToScreen,
 	const FGeometry& InAllottedGeometry, FSlateWindowElementList& OutDrawElements, int32 InLayerId,
 	EHitchDrawFlags InFlags
 	)
 {
 	using namespace Private;
 	const FDrawArgs InDrawArgs(InRangeToScreen, InAllottedGeometry, InLayerId, OutDrawElements);
+	const FHitchVisualizationStyle Style;
 	
 	if (EnumHasAnyFlags(InFlags, EHitchDrawFlags::DrawCatchupRanges))
 	{
-		DrawVerticalLines(InData.CatchupTimes, InDrawArgs);
-		DrawCatchupZones(InData.CatchupTimes, InDrawArgs);
+		DrawCatchupVerticalLines(InData.CatchupTimes, InHoverInfo.CatchupTimeIndex, Style.CatchupZoneStyle, InDrawArgs);
+		DrawCatchupZones(InData.CatchupTimes, InHoverInfo.CatchupTimeIndex, Style.CatchupZoneStyle, InDrawArgs);
 	}
 
 	if (EnumHasAnyFlags(InFlags, EHitchDrawFlags::DrawSkippedTimecodeMarkers))
 	{
-		DrawVerticalLines(InData.SkippedTimecodeMarkers, InDrawArgs);
-		DrawMarkerWarningIcons(InData.SkippedTimecodeMarkers, InDrawArgs);
+		DrawMarkerVerticalLines(InData.SkippedTimecodeMarkers, InHoverInfo.SkippedMarkerIndex, Style.SkipMarkerStyle, InDrawArgs);
+		DrawMarkerWarningIcons(InData.SkippedTimecodeMarkers, InHoverInfo.SkippedMarkerIndex, Style.SkipMarkerStyle, InDrawArgs);
 	}
 	
 	if (EnumHasAnyFlags(InFlags, EHitchDrawFlags::DrawRepeatedTimecodeMarkers))
 	{
-		DrawVerticalLines(InData.RepeatedTimecodeMarkers, InDrawArgs);
-		DrawMarkerWarningIcons(InData.RepeatedTimecodeMarkers, InDrawArgs);
+		DrawMarkerVerticalLines(InData.RepeatedTimecodeMarkers, InHoverInfo.RepeatedMarkerIndex, Style.RepeatMarkerStyle, InDrawArgs);
+		DrawMarkerWarningIcons(InData.RepeatedTimecodeMarkers, InHoverInfo.RepeatedMarkerIndex, Style.RepeatMarkerStyle, InDrawArgs);
 	}
 	
 	return InLayerId + 1;
 }
 
 int32 DrawTrackArea(
-	const FTimecodeHitchData& InData, const FScrubRangeToScreen& InRangeToScreen,
+	const FTimecodeHitchData& InData, const FHitchHoverInfo& InHoverInfo, const FScrubRangeToScreen& InRangeToScreen,
 	const FGeometry& InAllottedGeometry, FSlateWindowElementList& OutDrawElements, int32 InLayerId,
 	EHitchDrawFlags InFlags
 	)
 {
 	using namespace Private;
 	const FDrawArgs InDrawArgs(InRangeToScreen, InAllottedGeometry, InLayerId, OutDrawElements);
+	const FHitchVisualizationStyle Style;
 	
 	if (EnumHasAnyFlags(InFlags, EHitchDrawFlags::DrawCatchupRanges))
 	{
-		DrawVerticalLines(InData.CatchupTimes, InDrawArgs);
+		DrawCatchupVerticalLines(InData.CatchupTimes, InHoverInfo.CatchupTimeIndex, Style.CatchupZoneStyle, InDrawArgs);
 	}
 	
 	if (EnumHasAnyFlags(InFlags, EHitchDrawFlags::DrawSkippedTimecodeMarkers))
 	{
-		DrawVerticalLines(InData.SkippedTimecodeMarkers, InDrawArgs);
+		DrawMarkerVerticalLines(InData.SkippedTimecodeMarkers, InHoverInfo.SkippedMarkerIndex, Style.SkipMarkerStyle, InDrawArgs);
 	}
 
 	if (EnumHasAnyFlags(InFlags, EHitchDrawFlags::DrawRepeatedTimecodeMarkers))
 	{
-		DrawVerticalLines(InData.RepeatedTimecodeMarkers, InDrawArgs);
+		DrawMarkerVerticalLines(InData.RepeatedTimecodeMarkers, InHoverInfo.RepeatedMarkerIndex, Style.RepeatMarkerStyle, InDrawArgs);
 	}
 	
 	return InLayerId + 1;
 }
+
+namespace Private
+{
+enum class EHoverCheckFlags : uint8
+{
+	None,
+	Icon = 1 << 0,
+	Lines = 1 << 1,
+	All = Icon | Lines
+};
+ENUM_CLASS_FLAGS(EHoverCheckFlags);
+
+static bool IsLineHovered(
+	const FFrameTime InTime, const FVector2f& InAbsoluteCursorPos,
+	const FScrubRangeToScreen& InRangeToScreen, const FGeometry& InAllottedGeometry
+	)
+{
+	const float LinePosition = InRangeToScreen.FrameToLocalX(InTime);
+	
+	constexpr float HoverWidth = 2.f;
+	const float Left = LinePosition - HoverWidth;
+	const float Right = LinePosition + HoverWidth;
+	
+	const FVector2f Position(Left, 0);
+	const FVector2f Size(Right - Left, InAllottedGeometry.GetLocalSize().Y);
+	const FGeometry LineGeometry =  InAllottedGeometry.MakeChild(Size, FSlateLayoutTransform(Position));
+
+	return LineGeometry.IsUnderLocation(InAbsoluteCursorPos);
+}
+	
+static TOptional<int32> ComputeHoverState_Markers(
+	const TArray<FUnexpectedTimecodeMarker>& InMarkers, const FVector2f& InAbsoluteCursorPos,
+	const FScrubRangeToScreen& InRangeToScreen, const FGeometry& InAllottedGeometry, const FMarkerStyle& InStyle,
+	EHoverCheckFlags InFlags
+	)
+{
+	for (int32 Index = 0; Index < InMarkers.Num(); ++Index)
+	{
+		const FFrameTime& FrameTime = InMarkers[Index].Frame;
+		
+		if (EnumHasAnyFlags(InFlags, EHoverCheckFlags::Icon))
+		{
+			const FGeometry Geometry = MakeMarkerIconGeometry(FrameTime, InRangeToScreen, InAllottedGeometry, InStyle);
+			if (Geometry.IsUnderLocation(InAbsoluteCursorPos))
+			{
+				return Index;
+			}
+		}
+		
+		if (EnumHasAnyFlags(InFlags, EHoverCheckFlags::Lines)
+			&& IsLineHovered(FrameTime, InAbsoluteCursorPos, InRangeToScreen, InAllottedGeometry))
+		{
+			return Index;
+		}
+	}
+	return {};
+}
+
+static TOptional<int32> ComputeHoverState_Catchup(
+	const TArray<FCatchupTimeRange>& InTimes, const FVector2f& InAbsoluteCursorPos,
+	const FScrubRangeToScreen& InRangeToScreen, const FGeometry& InAllottedGeometry,
+	EHoverCheckFlags InFlags
+)
+{
+	for (int32 Index = 0; Index < InTimes.Num(); ++Index)
+	{
+		const FCatchupTimeRange& Zone = InTimes[Index];
+		
+		if (EnumHasAnyFlags(InFlags, EHoverCheckFlags::Icon))
+		{
+			const FGeometry Geometry = MakeCatchupZoneGeometry(Zone, InRangeToScreen, InAllottedGeometry);
+			if (Geometry.IsUnderLocation(InAbsoluteCursorPos))
+			{
+				return Index;
+			}
+		}
+		
+		if (EnumHasAnyFlags(InFlags, EHoverCheckFlags::Lines)
+			&& (IsLineHovered(Zone.StartTime, InAbsoluteCursorPos, InRangeToScreen, InAllottedGeometry)
+				|| IsLineHovered(GetCatchupRangeLastVerticalLineFrame(Zone), InAbsoluteCursorPos, InRangeToScreen, InAllottedGeometry)))
+		{
+			return Index;
+		}
+	}
+	return {};
+}
+
+static FHitchHoverInfo ComputeHoverState(
+	const FVector2f& InAbsoluteCursorPos, const FTimecodeHitchData& InData, const FScrubRangeToScreen& InRangeToScreen,
+	const FGeometry& InAllottedGeometry,
+	EHitchDrawFlags InFlags, EHoverCheckFlags InAreas
+	)
+{
+	using namespace Private;
+	
+	FHitchHoverInfo HoverInfo;
+	const FHitchVisualizationStyle Style;
+
+	if (EnumHasAnyFlags(InFlags, EHitchDrawFlags::DrawSkippedTimecodeMarkers))
+	{
+		HoverInfo.SkippedMarkerIndex = ComputeHoverState_Markers(
+			InData.SkippedTimecodeMarkers, InAbsoluteCursorPos, InRangeToScreen, InAllottedGeometry, Style.SkipMarkerStyle, InAreas
+			);
+	}
+	
+	if (EnumHasAnyFlags(InFlags, EHitchDrawFlags::DrawRepeatedTimecodeMarkers)
+		// Prefer hovering above items
+		&& !HoverInfo.IsHovered())
+	{
+		HoverInfo.RepeatedMarkerIndex = ComputeHoverState_Markers(
+			InData.RepeatedTimecodeMarkers, InAbsoluteCursorPos, InRangeToScreen, InAllottedGeometry, Style.RepeatMarkerStyle, InAreas
+			);
+	}
+	
+	if (EnumHasAnyFlags(InFlags, EHitchDrawFlags::DrawCatchupRanges)
+		// Prefer hovering above items
+		&& !HoverInfo.IsHovered())
+	{
+		HoverInfo.CatchupTimeIndex = ComputeHoverState_Catchup(
+			InData.CatchupTimes, InAbsoluteCursorPos, InRangeToScreen, InAllottedGeometry, InAreas
+			);
+	}
+
+	return HoverInfo;
+}
+}
+
+FHitchHoverInfo ComputeHoverStateForTimeSliderArea(
+	const FVector2f& InAbsoluteCursorPos,
+	const FTimecodeHitchData& InData, const FScrubRangeToScreen& InRangeToScreen,
+	const FGeometry& InAllottedGeometry, EHitchDrawFlags InFlags
+)
+{
+	const FHitchHoverInfo Info = ComputeHoverState(InAbsoluteCursorPos, InData, InRangeToScreen, InAllottedGeometry, InFlags,
+		Private::EHoverCheckFlags::All
+		);
+	return Info;
 }
+}
\ No newline at end of file
diff --git a/Engine/Plugins/VirtualProduction/Takes/Source/TakeRecorder/Private/Timecode/Visualization/Drawing/VisualizationDrawingUtils.h b/Engine/Plugins/VirtualProduction/Takes/Source/TakeRecorder/Private/Timecode/Visualization/Drawing/VisualizationDrawingUtils.h
--- a/Engine/Plugins/VirtualProduction/Takes/Source/TakeRecorder/Private/Timecode/Visualization/Drawing/VisualizationDrawingUtils.h
+++ b/Engine/Plugins/VirtualProduction/Takes/Source/TakeRecorder/Private/Timecode/Visualization/Drawing/VisualizationDrawingUtils.h
@@ -4,26 +4,36 @@
 
 #include "EHitchDrawFlags.h"
 #include "HAL/Platform.h"
+#include "Math/Vector2D.h"
 
 struct FGeometry;
 class FSlateWindowElementList;
 
 namespace UE::TakeRecorder
 {
+struct FHitchHoverInfo;
 struct FScrubRangeToScreen;
 struct FTimecodeHitchData;
 	
 /** Draws the top part (markers, etc.) of InData. */
 int32 DrawTimeSliderArea(
-	const FTimecodeHitchData& InData, const FScrubRangeToScreen& InRangeToScreen,
+	const FTimecodeHitchData& InData, const FHitchHoverInfo& InHoverInfo, const FScrubRangeToScreen& InRangeToScreen,
 	const FGeometry& InAllottedGeometry, FSlateWindowElementList& OutDrawElements, int32 InLayerId,
 	EHitchDrawFlags InFlags = EHitchDrawFlags::All
 	);
 	
 /** Draws the bottom part (markers, etc.) of InData. */
 int32 DrawTrackArea(
-	const FTimecodeHitchData& InData, const FScrubRangeToScreen& InRangeToScreen,
+	const FTimecodeHitchData& InData, const FHitchHoverInfo& InHoverInfo, const FScrubRangeToScreen& InRangeToScreen,
 	const FGeometry& InAllottedGeometry, FSlateWindowElementList& OutDrawElements, int32 InLayerId,
 	EHitchDrawFlags InFlags = EHitchDrawFlags::All
 	);
+
+
+/** Determines the hovered geometry in the time slider area. */
+FHitchHoverInfo ComputeHoverStateForTimeSliderArea(
+	const FVector2f& InAbsoluteCursorPos,
+	const FTimecodeHitchData& InData, const FScrubRangeToScreen& InRangeToScreen,
+	const FGeometry& InAllottedGeometry, EHitchDrawFlags InFlags = EHitchDrawFlags::All
+	);
 }
diff --git a/Engine/Plugins/VirtualProduction/Takes/Source/TakeRecorder/Private/Timecode/Visualization/HitchAnalysis.cpp b/Engine/Plugins/VirtualProduction/Takes/Source/TakeRecorder/Private/Timecode/Visualization/HitchAnalysis.cpp
--- a/Engine/Plugins/VirtualProduction/Takes/Source/TakeRecorder/Private/Timecode/Visualization/HitchAnalysis.cpp
+++ b/Engine/Plugins/VirtualProduction/Takes/Source/TakeRecorder/Private/Timecode/Visualization/HitchAnalysis.cpp
@@ -77,8 +77,8 @@ static FTimecode GetNextTimecode(const FTimecode& InCurrent, int32 InNumTimecode
 
 struct FAnalyzedMarkers
 {
-	TArray<FSkippedTimecodeMarker> Skipped;
-	TArray<FSkippedTimecodeMarker> Repeated;
+	TArray<FUnexpectedTimecodeMarker> Skipped;
+	TArray<FUnexpectedTimecodeMarker> Repeated;
 };
 	
 static FAnalyzedMarkers AnalyzeSkippedTimecodes(
@@ -100,11 +100,11 @@ static FAnalyzedMarkers AnalyzeSkippedTimecodes(
 			}
 
 			const bool bIsRepeat = LastTimecode && InActual == *LastTimecode;
-			TArray<FSkippedTimecodeMarker>& Markers = bIsRepeat ? Result.Repeated : Result.Skipped;
-			Markers.Add(FSkippedTimecodeMarker{ InCurrent });
+			TArray<FUnexpectedTimecodeMarker>& Markers = bIsRepeat ? Result.Repeated : Result.Skipped;
+			Markers.Add(FUnexpectedTimecodeMarker(InCurrent, InActual, *ExpectedTimecode));
 				
 			UE_CLOG(CVarLogFrameMarkers.GetValueOnAnyThread(), LogTemp, Warning, TEXT("%s timecode at %d. Actual: %s \tExpected: %s \tLast: %s"),
-				bIsRepeat ? TEXT("Repeated") : TEXT("Skipped"),InCurrent.Value, *InActual.ToString(), *ExpectedTimecode->ToString(),
+				bIsRepeat ? TEXT("Repeated") : TEXT("Skipped"), InCurrent.Value, *InActual.ToString(), *ExpectedTimecode->ToString(),
 				LastTimecode ? *LastTimecode->ToString() : TEXT("unset")
 				);
 		});
@@ -133,8 +133,8 @@ static void DumpCatchupRanges(UFrameHitchSceneDecoration& InDecoration, const TA
 		
 		UE_CLOG(CVarLogFrameCatchupRanges.GetValueOnAnyThread(), LogTemp, Warning,
 			TEXT("Catch-up range:\n\tStart at %d\tTarget TC: %s \tActual TC: %s\n\tEnd at %d   \tTarget TC: %s \tActual TC: %s"),
-			Range.StartTime.FrameNumber.Value, *Start_TargetFrameNumber->ToString(), *Start_ActualFrameNumber->ToString(),
-			Range.EndTime.FrameNumber.Value, *End_TargetFrameNumber->ToString(), *End_ActualFrameNumber->ToString()
+			Range.StartTime.Value, *Start_TargetFrameNumber->ToString(), *Start_ActualFrameNumber->ToString(),
+			Range.EndTime.Value, *End_TargetFrameNumber->ToString(), *End_ActualFrameNumber->ToString()
 		);
 	}
 }
@@ -153,7 +153,7 @@ static TArray<FCatchupTimeRange> AnalyzeCatchupZones(
 			const FFrameNumber ActualFrameNumber = InActual.ToFrameNumber(InSequencerRate);
 			if (!OpenCatchupBracket && ActualFrameNumber < TargetFrameNumber)
 			{
-				OpenCatchupBracket = FCatchupTimeRange{ InCurrent, InCurrent };
+				OpenCatchupBracket = FCatchupTimeRange(InCurrent, InCurrent);
 			}
 			else if (OpenCatchupBracket && ActualFrameNumber < TargetFrameNumber)
 			{
diff --git a/Engine/Plugins/VirtualProduction/Takes/Source/TakeRecorder/Private/Timecode/Visualization/HitchAnalysis.h b/Engine/Plugins/VirtualProduction/Takes/Source/TakeRecorder/Private/Timecode/Visualization/HitchAnalysis.h
--- a/Engine/Plugins/VirtualProduction/Takes/Source/TakeRecorder/Private/Timecode/Visualization/HitchAnalysis.h
+++ b/Engine/Plugins/VirtualProduction/Takes/Source/TakeRecorder/Private/Timecode/Visualization/HitchAnalysis.h
@@ -4,6 +4,7 @@
 
 #include "Misc/FrameTime.h"
 #include "Misc/Optional.h"
+#include "Misc/Timecode.h"
 
 class ISequencer;
 
@@ -13,9 +14,9 @@ namespace UE::TakeRecorder
 struct FCatchupTimeRange
 {
 	/** The first frame at which the timecode started being behind */
-	FFrameTime StartTime;
+	FFrameNumber StartTime;
 	/** The last frame at which the timecode was still behind */
-	FFrameTime EndTime;
+	FFrameNumber EndTime;
 
 	/**
 	 * The frame at which the timecode has caught up again.
@@ -23,24 +24,41 @@ struct FCatchupTimeRange
 	 *
 	 * Unset if the recording ended with a frame at which timecode was behind.
 	 */
-	TOptional<FFrameTime> FirstOkFrame;
+	TOptional<FFrameNumber> FirstOkFrame;
+
+	explicit FCatchupTimeRange(const FFrameNumber& InStartTime, const FFrameNumber& InEndTime, const TOptional<FFrameNumber>& InFirstOkFrame = {})
+		: StartTime(InStartTime)
+		, EndTime(InEndTime)
+		, FirstOkFrame(InFirstOkFrame)
+	{}
 };
 
 /** Data about a frame at which timecode was skipped */
-struct FSkippedTimecodeMarker
+struct FUnexpectedTimecodeMarker
 {
 	/** The frame at which the timecode did not match. */
-	FFrameTime Frame;
+	FFrameNumber Frame;
+
+	/** The actual timecode the frame had. */
+	FTimecode ActualTimecode;
+	/** The timecode that was expected at this frame. */
+	FTimecode ExpectedFrame;
+
+	explicit FUnexpectedTimecodeMarker(const FFrameNumber& InFrame, const FTimecode& InActualTimecode, const FTimecode& InExpectedFrame)
+		: Frame(InFrame)
+		, ActualTimecode(InActualTimecode)
+		, ExpectedFrame(InExpectedFrame)
+	{}
 };
 
 /** All data required to draw */
 struct FTimecodeHitchData
 {
 	/** Frames at which the timecode was skipped (previous frame was not current frame - 1). */
-	TArray<FSkippedTimecodeMarker> SkippedTimecodeMarkers;
+	TArray<FUnexpectedTimecodeMarker> SkippedTimecodeMarkers;
 
 	/** Frames at which timecode was repeated (previous frame had the same timecode). */
-	TArray<FSkippedTimecodeMarker> RepeatedTimecodeMarkers;
+	TArray<FUnexpectedTimecodeMarker> RepeatedTimecodeMarkers;
 
 	/** Sections at which the engine fell behind timecode. */
 	TArray<FCatchupTimeRange> CatchupTimes;
diff --git a/Engine/Plugins/VirtualProduction/Takes/Source/TakeRecorder/Private/Timecode/Visualization/Sequencer/HitchTooltipShower.cpp b/Engine/Plugins/VirtualProduction/Takes/Source/TakeRecorder/Private/Timecode/Visualization/Sequencer/HitchTooltipShower.cpp
new file mode 100644
--- /dev/null
+++ b/Engine/Plugins/VirtualProduction/Takes/Source/TakeRecorder/Private/Timecode/Visualization/Sequencer/HitchTooltipShower.cpp
@@ -0,0 +1,123 @@
+﻿// Copyright Epic Games, Inc. All Rights Reserved.
+
+#include "HitchTooltipShower.h"
+
+#include "Framework/Application/SlateApplication.h"
+#include "ISequencer.h"
+#include "Widgets/SToolTip.h"
+#include "Internationalization/Text.h"
+#include "Timecode/Visualization/HitchAnalysis.h"
+
+#define LOCTEXT_NAMESPACE "FHitchTooltipShower"
+
+namespace UE::TakeRecorder
+{
+namespace Private
+{
+static FText MakeHitchVisualizationTooltipText(
+	const ISequencer& InSequencer, const FTimecodeHitchData& HitchData, const FHitchHoverInfo& InHoverInfo
+	)
+{
+	const TSharedPtr<INumericTypeInterface<double>> NumericTypeInterface = InSequencer.GetNumericTypeInterface();
+	const auto ConvertFrameToText = [&NumericTypeInterface](const FFrameNumber& InTime)
+	{
+		const FString AsFrameString = ensure(NumericTypeInterface)
+			? NumericTypeInterface->ToString(InTime.Value)
+			: FString::FromInt(InTime.Value);
+		return FText::FromString(AsFrameString);
+	};
+
+	if (InHoverInfo.RepeatedMarkerIndex && ensure(HitchData.RepeatedTimecodeMarkers.IsValidIndex(*InHoverInfo.RepeatedMarkerIndex)))
+	{
+		const FText Fmt = LOCTEXT(
+			"RepeatedMarkerFmt",
+			"Timecode was repeated at frame {0}.\nActual timecode:{1}\nExpected timecode:{2}"
+			);
+		const FUnexpectedTimecodeMarker& Data = HitchData.RepeatedTimecodeMarkers[*InHoverInfo.RepeatedMarkerIndex];
+		return FText::Format(Fmt,
+			ConvertFrameToText(Data.Frame),
+			FText::FromString(Data.ActualTimecode.ToString()),
+			FText::FromString(Data.ExpectedFrame.ToString())
+			);
+	}
+	
+	if (InHoverInfo.SkippedMarkerIndex && ensure(HitchData.SkippedTimecodeMarkers.IsValidIndex(*InHoverInfo.SkippedMarkerIndex)))
+	{
+		const FText Fmt = LOCTEXT(
+			"SkippedMarkerFmt",
+			"A frame drop was detected at frame {0}.\nActual timecode:{1}\nExpected timecode:{2}\n\nData may be missing from record."
+			);
+		const FUnexpectedTimecodeMarker& Data = HitchData.SkippedTimecodeMarkers[*InHoverInfo.SkippedMarkerIndex];
+		return FText::Format(Fmt,
+			ConvertFrameToText(Data.Frame),
+			FText::FromString(Data.ActualTimecode.ToString()),
+			FText::FromString(Data.ExpectedFrame.ToString())
+			);
+	}
+	
+	if (InHoverInfo.CatchupTimeIndex && ensure(HitchData.CatchupTimes.IsValidIndex(*InHoverInfo.CatchupTimeIndex)))
+	{
+		const FText Fmt = LOCTEXT(
+			"CatchupZoneFmt",
+			"Take Recorder was unable to keep up with the target timestep.\n\nData may have been lost or misplotted in this range [{0} - {1}]."
+			);
+		const FCatchupTimeRange& Data = HitchData.CatchupTimes[*InHoverInfo.CatchupTimeIndex];
+		return FText::Format(Fmt, ConvertFrameToText(Data.StartTime), ConvertFrameToText(Data.EndTime));
+	}
+
+	return FText::GetEmpty();
+}
+}
+	
+FHitchTooltipShower::FTooltipState::FTooltipState(const FHitchHoverInfo& InHoverInfo, const FText& InTooltipText)
+	: HoverInfo(InHoverInfo)
+	, OurTooltip(SNew(SToolTip)
+		.Text(InTooltipText)
+		.IsInteractive(true) // We need to set the tooltip as interactive as otherwise the tooltip is immediately hidden next frame via widget search
+		.OnSetInteractiveWindowLocation_Lambda([](FVector2D& OutPosition){ OutPosition =  FSlateApplication::Get().GetCursorPos(); })
+		)
+{
+	FSlateApplication& SlateApplication = FSlateApplication::Get();
+	SlateApplication.SpawnToolTip(OurTooltip.ToSharedRef(), SlateApplication.GetCursorPos());
+}
+
+FHitchTooltipShower::FTooltipState::~FTooltipState()
+{
+	FSlateApplication& SlateApplication = FSlateApplication::Get();
+	SlateApplication.CloseToolTip();
+}
+
+bool FHitchTooltipShower::CanShowTooltips() const
+{
+	FSlateApplication& SlateApplication = FSlateApplication::Get();
+	return SlateApplication.GetAllowTooltips() && SlateApplication.GetPressedMouseButtons().IsEmpty();
+}
+	
+void FHitchTooltipShower::UpdateTooltipState(const FTimecodeHitchData& HitchData, const FHitchHoverInfo& InHoverInfo)
+{
+	const TSharedPtr<ISequencer> SequencerPin = WeakSequencer.Pin();
+	if (InHoverInfo && SequencerPin && CanShowTooltips())
+	{
+		RecreateIfContentChanged(HitchData, InHoverInfo, *SequencerPin);
+	}
+	else
+	{
+		TooltipState.Reset();
+	}
+}
+	
+void FHitchTooltipShower::RecreateIfContentChanged(const FTimecodeHitchData& HitchData, const FHitchHoverInfo& InHoverInfo, const ISequencer& InSequencer)
+{
+	const bool bNeedsTooltip = !TooltipState
+		|| InHoverInfo != TooltipState->HoverInfo;
+	
+	if (bNeedsTooltip)
+	{
+		TooltipState.Emplace(
+			InHoverInfo, Private::MakeHitchVisualizationTooltipText(InSequencer, HitchData, InHoverInfo)
+		);
+	}
+}
+}
+
+#undef LOCTEXT_NAMESPACE
\ No newline at end of file
diff --git a/Engine/Plugins/VirtualProduction/Takes/Source/TakeRecorder/Private/Timecode/Visualization/Sequencer/HitchTooltipShower.h b/Engine/Plugins/VirtualProduction/Takes/Source/TakeRecorder/Private/Timecode/Visualization/Sequencer/HitchTooltipShower.h
new file mode 100644
--- /dev/null
+++ b/Engine/Plugins/VirtualProduction/Takes/Source/TakeRecorder/Private/Timecode/Visualization/Sequencer/HitchTooltipShower.h
@@ -0,0 +1,65 @@
+﻿// Copyright Epic Games, Inc. All Rights Reserved.
+
+#pragma once
+
+#if WITH_EDITOR
+
+#include "Misc/Optional.h"
+#include "Templates/SharedPointer.h"
+#include "Templates/UnrealTemplate.h"
+#include "Timecode/Visualization/Drawing/HitchHoverInfo.h"
+
+class ISequencer;
+class IToolTip;
+struct FFrameTime;
+
+namespace UE::TakeRecorder
+{
+struct FTimecodeHitchData;
+
+/** Responsible for generating a tooltip based on the currently hovered visualization UI. */
+class FHitchTooltipShower : public FNoncopyable
+{
+public:
+
+	explicit FHitchTooltipShower(const TWeakPtr<ISequencer> InWeakSequencer)
+		: WeakSequencer(InWeakSequencer)
+	{}
+	
+	/** @return Whether tooltips are allowed to be shown currently. */
+	bool CanShowTooltips() const;
+	
+	/** Shows a visualization tooltip if currently hovering */
+	void UpdateTooltipState(const FTimecodeHitchData& HitchData, const FHitchHoverInfo& InHoverInfo);
+
+	/** Ensures the tooltip is hidden. */
+	void ClearTooltip()
+	{
+		TooltipState.Reset();
+	}
+
+private:
+
+	/** Used to convert frame times to display strings. */
+	const TWeakPtr<ISequencer> WeakSequencer;
+	
+	struct FTooltipState
+	{
+		/** The hovered UI elements when the tooltip was overriden. */
+		const FHitchHoverInfo HoverInfo;
+
+		/** The hitch visualization tooltip we've temporarily spawned. */
+		const TSharedPtr<IToolTip> OurTooltip;
+
+		explicit FTooltipState(const FHitchHoverInfo& InHoverInfo, const FText& InTooltipTex);
+		~FTooltipState();
+	};
+
+	/** Active tooltip override. Cleans up the tooltip on destruction. */
+	TOptional<FTooltipState> TooltipState;
+
+	/** (Re-)creates the tooltip if the hover a UI element is hovered. If the hovered UI has changed since last time, the tooltip is recreated. */
+	void RecreateIfContentChanged(const FTimecodeHitchData& HitchData, const FHitchHoverInfo& InHoverInfo, const ISequencer& InSequence);
+};
+}
+#endif
\ No newline at end of file
diff --git a/Engine/Plugins/VirtualProduction/Takes/Source/TakeRecorder/Private/Timecode/Visualization/Sequencer/HitchVisualizationViewModel.cpp b/Engine/Plugins/VirtualProduction/Takes/Source/TakeRecorder/Private/Timecode/Visualization/Sequencer/HitchVisualizationViewModel.cpp
--- a/Engine/Plugins/VirtualProduction/Takes/Source/TakeRecorder/Private/Timecode/Visualization/Sequencer/HitchVisualizationViewModel.cpp
+++ b/Engine/Plugins/VirtualProduction/Takes/Source/TakeRecorder/Private/Timecode/Visualization/Sequencer/HitchVisualizationViewModel.cpp
@@ -5,9 +5,25 @@
 #include "Timecode/Visualization/Drawing/ScrubRangeToScreen.h"
 #include "Timecode/Visualization/Drawing/VisualizationDrawingUtils.h"
 #include "Timecode/Visualization/TakeRecorderHitchVisualizationSettings.h"
+#include "Widgets/SToolTip.h"
 
 namespace UE::TakeRecorder
 {
+namespace Private
+{
+static FHitchHoverInfo CombineHoverInfos(const FHitchHoverInfo& InFirst, const FHitchHoverInfo& InSecond)
+{
+	// *Technically* both options could be set but that should not be the case because the mouse is only under a single location.
+	// Even if they differed, we wouldn't care... it's just visual.
+	return
+	{
+		InFirst.SkippedMarkerIndex ? InFirst.SkippedMarkerIndex : InSecond.SkippedMarkerIndex,
+		InFirst.RepeatedMarkerIndex ? InFirst.RepeatedMarkerIndex : InSecond.RepeatedMarkerIndex,
+		InFirst.CatchupTimeIndex ? InFirst.CatchupTimeIndex : InSecond.CatchupTimeIndex,
+	};
+}
+}
+	
 UE_SEQUENCER_DEFINE_CASTABLE(FHitchVisualizationViewModel);
 
 int32 FHitchVisualizationViewModel::OnPaintTimeSliderOverlay(
@@ -17,18 +33,22 @@ int32 FHitchVisualizationViewModel::OnPaintTimeSliderOverlay(
 {
 	if (!CanPaint())
 	{
+		TooltipShowerHack.ClearTooltip();
 		return LayerId;
 	}
 	
 	UTakeRecorderHitchVisualizationSettings* Settings = UTakeRecorderHitchVisualizationSettings::Get();
+	UpdateHoverState_TimeSlider(Args, AllottedGeometry, Settings->GetFlagsForUI());
 	return DrawTimeSliderArea(
-		HitchData, MakeScrubRangeToScreen(AllottedGeometry), AllottedGeometry, OutDrawElements, LayerId, Settings->GetFlagsForUI()
+		HitchData, TimeSliderHoverInfo,
+		MakeScrubRangeToScreen(AllottedGeometry), AllottedGeometry, OutDrawElements, LayerId, Settings->GetFlagsForUI()
 		);
 }
 
 int32 FHitchVisualizationViewModel::OnPaintTrackAreaOverlay(
 	const FPaintArgs& Args, const FGeometry& AllottedGeometry, const FSlateRect& MyCullingRect, FSlateWindowElementList& OutDrawElements,
-	int32 LayerId, const FWidgetStyle& InWidgetStyle)
+	int32 LayerId, const FWidgetStyle& InWidgetStyle
+	)
 {
 	if (!CanPaint())
 	{
@@ -37,10 +57,38 @@ int32 FHitchVisualizationViewModel::OnPaintTrackAreaOverlay(
 	
 	UTakeRecorderHitchVisualizationSettings* Settings = UTakeRecorderHitchVisualizationSettings::Get();
 	return DrawTrackArea(
-		HitchData, MakeScrubRangeToScreen(AllottedGeometry), AllottedGeometry, OutDrawElements, LayerId, Settings->GetFlagsForUI()
+		HitchData, TimeSliderHoverInfo,
+		MakeScrubRangeToScreen(AllottedGeometry), AllottedGeometry, OutDrawElements, LayerId, Settings->GetFlagsForUI()
 		);
 }
 
+	
+void FHitchVisualizationViewModel::UpdateHoverState_TimeSlider(
+	const FPaintArgs& InArgs, const FGeometry& InAllottedGeometry, EHitchDrawFlags InFlags
+	)
+{
+	// Don't draw the UI as hovered while the user is dragging the time slider.
+	const bool bCanHover = FSlateApplication::Get().GetPressedMouseButtons().IsEmpty();
+	if (bCanHover)
+	{
+		// InAllottedGeometry is in paint space. Mouse operations, etc. need to be done in tick space.
+		// The difference is that in paint space everything is shifted by GetWindowToDesktopTransform... so we need to undo that shift.
+		FGeometry TickSpaceGeometry = InAllottedGeometry;
+		TickSpaceGeometry.AppendTransform(FSlateLayoutTransform(InArgs.GetWindowToDesktopTransform()));
+		
+		TimeSliderHoverInfo = ComputeHoverStateForTimeSliderArea(
+				FSlateApplication::Get().GetCursorPos(), HitchData, MakeScrubRangeToScreen(TickSpaceGeometry), TickSpaceGeometry, InFlags
+			);
+	}
+	else
+	{
+		TimeSliderHoverInfo.Reset();
+	}
+
+	// As we're manually managing the tooltip state, reevaluate whether it should be shown based off of our new hover state.
+	TooltipShowerHack.UpdateTooltipState(HitchData, TimeSliderHoverInfo);
+}
+
 FScrubRangeToScreen FHitchVisualizationViewModel::MakeScrubRangeToScreen(const FGeometry& AllottedGeometry) const
 {
 	const TSharedPtr<ISequencer> SequencerPin = WeakOwningSequencer.Pin();
diff --git a/Engine/Plugins/VirtualProduction/Takes/Source/TakeRecorder/Private/Timecode/Visualization/Sequencer/HitchVisualizationViewModel.h b/Engine/Plugins/VirtualProduction/Takes/Source/TakeRecorder/Private/Timecode/Visualization/Sequencer/HitchVisualizationViewModel.h
--- a/Engine/Plugins/VirtualProduction/Takes/Source/TakeRecorder/Private/Timecode/Visualization/Sequencer/HitchVisualizationViewModel.h
+++ b/Engine/Plugins/VirtualProduction/Takes/Source/TakeRecorder/Private/Timecode/Visualization/Sequencer/HitchVisualizationViewModel.h
@@ -3,17 +3,35 @@
 #pragma once
 
 #if WITH_EDITOR
+#include "HitchTooltipShower.h"
 #include "MVVM/Extensions/ITopTimeSliderOverlayExtension.h"
 #include "MVVM/Extensions/ITrackAreaOverlayExtension.h"
 #include "MVVM/ViewModels/ViewModel.h"
 #include "Misc/Attribute.h"
+#include "Timecode/Visualization/Drawing/HitchHoverInfo.h"
 #include "Timecode/Visualization/HitchAnalysis.h"
+#include "Timecode/Visualization/TakeRecorderHitchVisualizationSettings.h"
+
+class IToolTip;
 
 namespace UE::TakeRecorder
 {
 struct FScrubRangeToScreen;
 	
-/**  */
+/**
+ * Draws the hitch visualization UI:
+ * - Timecode skip & repeats draw a vertical line and a warning icon (in red ),
+ * - Catch up zones, when the engine was behind the target timecode (in yellow),
+ * - Draws a tooltip when UI in the timeslider section is hovered. The content is also drawn brighter.
+ * - Content hovered in the track area cannot be hovered. The reason for this is because Sequencer's widgets, such as tracks, get precedence.
+ * so you can get into a state where you hover the hitch vertical line, it is drawn brighter, but it shows the tooltip of the track instead of
+ * the hitch tooltip.
+ * That's the reason why we need to update Sequencer's hotspot API to support tooltips (and extend it to include time slider area as well.
+ * @see TooltipShowerHack.
+ *
+ * This is set up by FSequencerHitchVisualizer.
+ * @see FSequencerHitchVisualizer
+ */
 class FHitchVisualizationViewModel
 	: public Sequencer::FViewModel
 	, public Sequencer::ITopTimeSliderOverlayExtension
@@ -29,6 +47,7 @@ public:
 		: WeakOwningSequencer(MoveTemp(InWeakOwningSequencer))
 		, CanPaintAttr(MoveTemp(InCanPaintAttr))
 		, HitchData(MoveTemp(InHitchData))
+		, TooltipShowerHack(WeakOwningSequencer)
 	{}
 	
 	//~ Begin ITimeSliderOverlayExtension Interface
@@ -49,7 +68,24 @@ private:
 
 	/** The hitch data that is being visualized. */
 	FTimecodeHitchData HitchData;
+	/** In time slider area, info about which of HitchData should be drawn has hovered. Also used for determining what tooltip to show. */
+	FHitchHoverInfo TimeSliderHoverInfo;
+
+	/**
+	 * Handles showing and hiding the tooltip for our hitch visualization UI.
+	 *
+	 * This is a "hack" because we should prefer be using Sequencer's hotspot API: ideally, we'd tell the hotspot API what the user is currently
+	 * focusing. Then, the API would ask what the tooltip should be and handle showing it. However, the API is limited:
+	 * - the hotspot API is limited to track areas but we need it for the time slider area as well, and
+	 * - the hotspot API provides not way to show tooltips for now.
+	 * Once the hotspot API supports those operations, this hack should be replaced.
+	 */
+	FHitchTooltipShower TooltipShowerHack;
+
+	/** Updates TimeSliderHoverInfo and the resulting tooltip state.*/
+	void UpdateHoverState_TimeSlider(const FPaintArgs& InArgs, const FGeometry& InAllottedGeometry, EHitchDrawFlags InFlags);
 
+	/** @return Utility for converting FFrameTime to widget paint positions. */
 	FScrubRangeToScreen MakeScrubRangeToScreen(const FGeometry& AllottedGeometry) const;
 
 	/** @return Whether this model should paint itself. */
