Take Recorder: When FPS mismatch, hitch visualization shows a warning icon in track area.

diff --git a/Engine/Plugins/VirtualProduction/Takes/Source/TakeMovieScene/Public/Hitching/FrameHitchSceneDecoration.h b/Engine/Plugins/VirtualProduction/Takes/Source/TakeMovieScene/Public/Hitching/FrameHitchSceneDecoration.h
--- a/Engine/Plugins/VirtualProduction/Takes/Source/TakeMovieScene/Public/Hitching/FrameHitchSceneDecoration.h
+++ b/Engine/Plugins/VirtualProduction/Takes/Source/TakeMovieScene/Public/Hitching/FrameHitchSceneDecoration.h
@@ -26,10 +26,17 @@ public:
 	UPROPERTY()
 	FTakeMovieSceneHitchTimecodeCurves ActualTimecode;
 
-	/** The frame rate at which was recorded. Determines the number of frames per second in the FTimecode. */
+	/** The frame rate the timecode provider was set up with. Determines the number of frames per second in TargetTimecode & ActualTimecode. */
 	UPROPERTY()
-	FFrameRate FrameRate;
-
+	FFrameRate TimecodeProviderFrameRate;
+	
+	/**
+	 * The frame rate Take Recorder was set up with. Users may (intentionally) set up the provider and sequencer to record at different rates.
+	 * This information is important for analyzing hitches later as it affects the timecode we should expect each frame.
+	 */
+	UPROPERTY()
+	FFrameRate RecordFrameRate;
+	
 	/** Evaluates the curve data at the specified time. */
 	TAKEMOVIESCENE_API TOptional<UE::TakeMovieScene::FFrameHitchData> Evaluate(const FFrameTime& InFrameTime) const;
 };
diff --git a/Engine/Plugins/VirtualProduction/Takes/Source/TakeRecorder/Private/TakeRecorderStyle.cpp b/Engine/Plugins/VirtualProduction/Takes/Source/TakeRecorder/Private/TakeRecorderStyle.cpp
--- a/Engine/Plugins/VirtualProduction/Takes/Source/TakeRecorder/Private/TakeRecorderStyle.cpp
+++ b/Engine/Plugins/VirtualProduction/Takes/Source/TakeRecorder/Private/TakeRecorderStyle.cpp
@@ -226,6 +226,11 @@ FTakeRecorderStyle::FTakeRecorderStyle()
 	Set("Hitching.CatchupZone.Normal", WarningColor); 
 	Set("Hitching.CatchupZone.Hovered", WarningColor_Hovered); 
 	Set("Hitching.CatchupZone.Icon", new FSlateColorBrush(FLinearColor::White)); // The colour does not matter - we'll tint it in the drawing code.
+
+	Set("Hitching.MismatchedFrameRate.Normal", WarningColor); 
+	Set("Hitching.MismatchedFrameRate.Hovered", WarningColor_Hovered); 
+	Set("Hitching.MismatchedFrameRate.Icon", new IMAGE_BRUSH_SVG("AlertTriangleSolid", Icon16x16));
+	Set("Hitching.MismatchedFrameRate.Offset", FVector2f(4.f, 4.f));
 	
 	FSlateStyleRegistry::RegisterSlateStyle(*this);
 }
diff --git a/Engine/Plugins/VirtualProduction/Takes/Source/TakeRecorder/Private/Timecode/Recording/FrameHitchDecorationRecorder.cpp b/Engine/Plugins/VirtualProduction/Takes/Source/TakeRecorder/Private/Timecode/Recording/FrameHitchDecorationRecorder.cpp
--- a/Engine/Plugins/VirtualProduction/Takes/Source/TakeRecorder/Private/Timecode/Recording/FrameHitchDecorationRecorder.cpp
+++ b/Engine/Plugins/VirtualProduction/Takes/Source/TakeRecorder/Private/Timecode/Recording/FrameHitchDecorationRecorder.cpp
@@ -8,6 +8,7 @@
 #include "LevelSequence.h"
 #include "MovieScene.h"
 #include "Recorder/TakeRecorder.h"
+#include "Recorder/TakeRecorderSubsystem.h"
 #include "TakeRecorderSources.h"
 #include "TakesCoreLog.h"
 
@@ -125,7 +126,11 @@ void FFrameHitchDecorationRecorder::OnRecordingStopped(UTakeRecorder* TakeRecord
 	Decoration->ActualTimecode.Set(
 		MoveTemp(ActualData.Times), MoveTemp(ActualData.Hours), MoveTemp(ActualData.Minutes), MoveTemp(ActualData.Seconds), MoveTemp(ActualData.Frames)
 		);
-	Decoration->FrameRate = TargetData.TCRate; 
+	Decoration->TimecodeProviderFrameRate = TargetData.TCRate;
+	
+	// We don't expect the target frame rate to have been changed since the start of the recording because the UI disables the frame rate from being edited.
+	// Should that change in the future, simply set this when recording starts.
+	Decoration->RecordFrameRate = GEngine->GetEngineSubsystem<UTakeRecorderSubsystem>()->GetFrameRate();
 }
 }
 
diff --git a/Engine/Plugins/VirtualProduction/Takes/Source/TakeRecorder/Private/Timecode/Visualization/Drawing/AnalyzedHitchUIDrawingUtils.cpp b/Engine/Plugins/VirtualProduction/Takes/Source/TakeRecorder/Private/Timecode/Visualization/Drawing/AnalyzedHitchUIDrawingUtils.cpp
new file mode 100644
--- /dev/null
+++ b/Engine/Plugins/VirtualProduction/Takes/Source/TakeRecorder/Private/Timecode/Visualization/Drawing/AnalyzedHitchUIDrawingUtils.cpp
@@ -0,0 +1,366 @@
+﻿// Copyright Epic Games, Inc. All Rights Reserved.
+
+#include "AnalyzedHitchUIDrawingUtils.h"
+
+#include "EHitchDrawFlags.h"
+#include "HitchVisualizationStyle.h"
+#include "Misc/FrameTime.h"
+#include "Rendering/DrawElementTypes.h"
+#include "ScrubRangeToScreen.h"
+#include "Timecode/Visualization/Drawing/AnalyzedHitchUIHoverInfo.h"
+#include "Timecode/Visualization/HitchAnalysis.h"
+
+namespace UE::TakeRecorder::HitchVisualizationUI
+{
+namespace Private
+{
+/** Packs args that are used together for drawing geometry in Sequencer's widget overlay areas. */
+struct FDrawArgs
+{
+	/** Converts FFrameTime to Sequencer's widget area. */
+	const FScrubRangeToScreen& RangeToScreen;
+	/** The geometry area allocated to the overlay */
+	const FGeometry& AllottedGeometry;
+	/** The base layer ID to draw onto */
+	int32 LayerId;
+
+	/** The list to append draw elements to */
+	FSlateWindowElementList& OutDrawElements;
+
+	explicit FDrawArgs(const FScrubRangeToScreen& RangeToScreen, const FGeometry& AllottedGeometry, int32 LayerId, FSlateWindowElementList& OutDrawElements)
+		: RangeToScreen(RangeToScreen)
+		, AllottedGeometry(AllottedGeometry)
+		, LayerId(LayerId)
+		, OutDrawElements(OutDrawElements)
+	{}
+};
+
+static FGeometry MakeMarkerIconGeometry(
+	const FFrameTime& InTime, const FScrubRangeToScreen& InRangeToScreen, const FGeometry& InAllottedGeometry, const FMarkerStyle& InStyle
+	)
+{
+	const float MarkerPos = InRangeToScreen.FrameToLocalX(InTime);
+	const FVector2f Position = FVector2f(MarkerPos, 0) + InStyle.MarkerBrushOffset;
+	const FVector2f Size = InStyle.IconBrush->GetImageSize();
+	return InAllottedGeometry.MakeChild(Size, FSlateLayoutTransform(Position));
+}
+
+/**
+ * @return The frame at which to paint the range's right vertical line.
+ * 
+ * The line should be drawn at the beginning of the first frame after the range. Suppose we had a catchup range consisting of a single frame.
+ * Then StartTime and EndTime would be equal ... we wouldn't want to draw the left and right vertical lines at the same spot, i.e. a single line.
+ */
+static FFrameTime GetCatchupRangeLastVerticalLineFrame(const FCatchupTimeRange& InTimeRange)
+{
+	return InTimeRange.FirstOkFrame ? *InTimeRange.FirstOkFrame : InTimeRange.EndTime;
+}
+
+static FGeometry MakeCatchupZoneGeometry(
+	const FCatchupTimeRange& InTimeRange, const FScrubRangeToScreen& InRangeToScreen, const FGeometry& InAllottedGeometry
+	)
+{
+	const float Left = InRangeToScreen.FrameToLocalX(InTimeRange.StartTime);
+	const float Right = InRangeToScreen.FrameToLocalX(GetCatchupRangeLastVerticalLineFrame(InTimeRange));
+	const FVector2f Position(Left, 0);
+	const FVector2f Size(Right - Left, InAllottedGeometry.GetLocalSize().Y);
+	return InAllottedGeometry.MakeChild(Size, FSlateLayoutTransform(Position));
+}
+	
+static void DrawVerticalLine(
+	const FFrameTime& InFrameTime, const FDrawArgs& InDrawArgs, const FLinearColor& InColor
+	)
+{
+	const float LinePos = InDrawArgs.RangeToScreen.FrameToLocalX(InFrameTime);
+	
+	TArray<FVector2D> LinePoints;
+	LinePoints.AddUninitialized(2);
+	LinePoints[0] = FVector2D(LinePos, 0.0f);
+	LinePoints[1] = FVector2D(LinePos, FMath::FloorToFloat(InDrawArgs.AllottedGeometry.Size.Y));
+
+	// Avoid painting over the tree view on the left, or the dockable tabs on the right.
+	const bool bOutOfBounds = LinePos < 0.f || LinePos >= InDrawArgs.AllottedGeometry.Size.X;
+	if (bOutOfBounds)
+	{
+		return;
+	}
+	
+	FSlateDrawElement::MakeLines(
+		InDrawArgs.OutDrawElements,
+		InDrawArgs.LayerId,
+		InDrawArgs.AllottedGeometry.ToPaintGeometry(),
+		LinePoints,
+		ESlateDrawEffect::None,
+		InColor,
+		false,
+		1.f
+	);
+}
+	
+static void DrawMarkerVerticalLines(
+	TConstArrayView<FUnexpectedTimecodeMarker> InMarkers, const TOptional<int32>& InHoveredIndex, const FMarkerStyle& InStyle,
+	const FDrawArgs& InDrawArgs
+	)
+{
+	for (int32 Index = 0; Index < InMarkers.Num(); ++Index)
+	{
+		const FUnexpectedTimecodeMarker& Marker = InMarkers[Index];
+		const bool bIsHovered = InHoveredIndex && *InHoveredIndex == Index;
+		DrawVerticalLine(Marker.Frame, InDrawArgs, bIsHovered ? InStyle.HoveredTint : InStyle.NormalTint);
+	}
+}
+
+static void DrawCatchupVerticalLines(
+	TConstArrayView<FCatchupTimeRange> InTimes, const TOptional<int32>& InHoveredIndex, const FCatchupZoneStyle& InStyle,
+	const FDrawArgs& InDrawArgs
+	)
+{
+	for (int32 Index = 0; Index < InTimes.Num(); ++Index)
+	{
+		const FCatchupTimeRange& CatchupRange = InTimes[Index];
+		const bool bIsHovered = InHoveredIndex && *InHoveredIndex == Index;
+		const FLinearColor& InColor = bIsHovered ? InStyle.CatchupZone_Hovered : InStyle.CatchupZone_Normal;
+		
+		DrawVerticalLine(CatchupRange.StartTime, InDrawArgs, InColor);
+		DrawVerticalLine(GetCatchupRangeLastVerticalLineFrame(CatchupRange), InDrawArgs, InColor);
+	}
+}
+	
+static void DrawMarkerWarningIcons(
+	TConstArrayView<FUnexpectedTimecodeMarker> InMarkers, const TOptional<int32>& InHoveredIndex, const FMarkerStyle& InStyle,
+	const FDrawArgs& InDrawArgs
+	)
+{
+	for (int32 Index = 0; Index < InMarkers.Num(); ++Index)
+	{
+		const FUnexpectedTimecodeMarker& Marker = InMarkers[Index];
+		const bool bIsHovered = InHoveredIndex && *InHoveredIndex == Index;
+		FSlateDrawElement::MakeBox(
+			InDrawArgs.OutDrawElements,
+			InDrawArgs.LayerId,
+			MakeMarkerIconGeometry(Marker.Frame, InDrawArgs.RangeToScreen, InDrawArgs.AllottedGeometry, InStyle).ToPaintGeometry(),
+			InStyle.IconBrush,
+			ESlateDrawEffect::None,
+			bIsHovered ? InStyle.HoveredTint : InStyle.NormalTint
+			);
+	}
+}
+
+static void DrawCatchupZones(
+	TConstArrayView<FCatchupTimeRange> InTimes,  const TOptional<int32>& InHoveredIndex, const FCatchupZoneStyle& InStyle,
+	const FDrawArgs& InDrawArgs
+	)
+{
+	for (int32 Index = 0; Index < InTimes.Num(); ++Index)
+	{
+		const FCatchupTimeRange& CatchupRange = InTimes[Index];
+		const bool bIsHovered = InHoveredIndex && *InHoveredIndex == Index;
+		FSlateDrawElement::MakeBox(
+			InDrawArgs.OutDrawElements,
+			InDrawArgs.LayerId,
+			MakeCatchupZoneGeometry(CatchupRange, InDrawArgs.RangeToScreen, InDrawArgs.AllottedGeometry).ToPaintGeometry(),
+			InStyle.CatchupZoneBrush,
+			ESlateDrawEffect::None,
+			bIsHovered ? InStyle.CatchupZone_Hovered : InStyle.CatchupZone_Normal
+			);
+	}
+}
+}
+	
+int32 DrawTimeSliderArea(
+	const FTimecodeHitchData& InData, const FAnalyzedHitchUIHoverInfo& InHoverInfo, const FScrubRangeToScreen& InRangeToScreen,
+	const FGeometry& InAllottedGeometry, FSlateWindowElementList& OutDrawElements, int32 InLayerId,
+	EHitchDrawFlags InFlags
+	)
+{
+	using namespace Private;
+	const FDrawArgs InDrawArgs(InRangeToScreen, InAllottedGeometry, InLayerId, OutDrawElements);
+	const FHitchVisualizationStyle Style;
+	
+	if (EnumHasAnyFlags(InFlags, EHitchDrawFlags::DrawCatchupRanges))
+	{
+		DrawCatchupVerticalLines(InData.CatchupTimes, InHoverInfo.CatchupTimeIndex, Style.CatchupZoneStyle, InDrawArgs);
+		DrawCatchupZones(InData.CatchupTimes, InHoverInfo.CatchupTimeIndex, Style.CatchupZoneStyle, InDrawArgs);
+	}
+
+	if (EnumHasAnyFlags(InFlags, EHitchDrawFlags::DrawSkippedTimecodeMarkers))
+	{
+		DrawMarkerVerticalLines(InData.SkippedTimecodeMarkers, InHoverInfo.SkippedMarkerIndex, Style.SkipMarkerStyle, InDrawArgs);
+		DrawMarkerWarningIcons(InData.SkippedTimecodeMarkers, InHoverInfo.SkippedMarkerIndex, Style.SkipMarkerStyle, InDrawArgs);
+	}
+	
+	if (EnumHasAnyFlags(InFlags, EHitchDrawFlags::DrawRepeatedTimecodeMarkers))
+	{
+		DrawMarkerVerticalLines(InData.RepeatedTimecodeMarkers, InHoverInfo.RepeatedMarkerIndex, Style.RepeatMarkerStyle, InDrawArgs);
+		DrawMarkerWarningIcons(InData.RepeatedTimecodeMarkers, InHoverInfo.RepeatedMarkerIndex, Style.RepeatMarkerStyle, InDrawArgs);
+	}
+	
+	return InLayerId + 1;
+}
+
+int32 DrawTrackArea(
+	const FTimecodeHitchData& InData, const FAnalyzedHitchUIHoverInfo& InHoverInfo, const FScrubRangeToScreen& InRangeToScreen,
+	const FGeometry& InAllottedGeometry, FSlateWindowElementList& OutDrawElements, int32 InLayerId,
+	EHitchDrawFlags InFlags
+	)
+{
+	using namespace Private;
+	const FDrawArgs InDrawArgs(InRangeToScreen, InAllottedGeometry, InLayerId, OutDrawElements);
+	const FHitchVisualizationStyle Style;
+	
+	if (EnumHasAnyFlags(InFlags, EHitchDrawFlags::DrawCatchupRanges))
+	{
+		DrawCatchupVerticalLines(InData.CatchupTimes, InHoverInfo.CatchupTimeIndex, Style.CatchupZoneStyle, InDrawArgs);
+	}
+	
+	if (EnumHasAnyFlags(InFlags, EHitchDrawFlags::DrawSkippedTimecodeMarkers))
+	{
+		DrawMarkerVerticalLines(InData.SkippedTimecodeMarkers, InHoverInfo.SkippedMarkerIndex, Style.SkipMarkerStyle, InDrawArgs);
+	}
+
+	if (EnumHasAnyFlags(InFlags, EHitchDrawFlags::DrawRepeatedTimecodeMarkers))
+	{
+		DrawMarkerVerticalLines(InData.RepeatedTimecodeMarkers, InHoverInfo.RepeatedMarkerIndex, Style.RepeatMarkerStyle, InDrawArgs);
+	}
+	
+	return InLayerId + 1;
+}
+
+namespace Private
+{
+enum class EHoverCheckFlags : uint8
+{
+	None,
+	Icon = 1 << 0,
+	Lines = 1 << 1,
+	All = Icon | Lines
+};
+ENUM_CLASS_FLAGS(EHoverCheckFlags);
+
+static bool IsLineHovered(
+	const FFrameTime InTime, const FVector2f& InAbsoluteCursorPos,
+	const FScrubRangeToScreen& InRangeToScreen, const FGeometry& InAllottedGeometry
+	)
+{
+	const float LinePosition = InRangeToScreen.FrameToLocalX(InTime);
+	
+	constexpr float HoverWidth = 2.f;
+	const float Left = LinePosition - HoverWidth;
+	const float Right = LinePosition + HoverWidth;
+	
+	const FVector2f Position(Left, 0);
+	const FVector2f Size(Right - Left, InAllottedGeometry.GetLocalSize().Y);
+	const FGeometry LineGeometry =  InAllottedGeometry.MakeChild(Size, FSlateLayoutTransform(Position));
+
+	return LineGeometry.IsUnderLocation(InAbsoluteCursorPos);
+}
+	
+static TOptional<int32> ComputeHoverState_Markers(
+	const TArray<FUnexpectedTimecodeMarker>& InMarkers, const FVector2f& InAbsoluteCursorPos,
+	const FScrubRangeToScreen& InRangeToScreen, const FGeometry& InAllottedGeometry, const FMarkerStyle& InStyle,
+	EHoverCheckFlags InFlags
+	)
+{
+	for (int32 Index = 0; Index < InMarkers.Num(); ++Index)
+	{
+		const FFrameTime& FrameTime = InMarkers[Index].Frame;
+		
+		if (EnumHasAnyFlags(InFlags, EHoverCheckFlags::Icon))
+		{
+			const FGeometry Geometry = MakeMarkerIconGeometry(FrameTime, InRangeToScreen, InAllottedGeometry, InStyle);
+			if (Geometry.IsUnderLocation(InAbsoluteCursorPos))
+			{
+				return Index;
+			}
+		}
+		
+		if (EnumHasAnyFlags(InFlags, EHoverCheckFlags::Lines)
+			&& IsLineHovered(FrameTime, InAbsoluteCursorPos, InRangeToScreen, InAllottedGeometry))
+		{
+			return Index;
+		}
+	}
+	return {};
+}
+
+static TOptional<int32> ComputeHoverState_Catchup(
+	const TArray<FCatchupTimeRange>& InTimes, const FVector2f& InAbsoluteCursorPos,
+	const FScrubRangeToScreen& InRangeToScreen, const FGeometry& InAllottedGeometry,
+	EHoverCheckFlags InFlags
+)
+{
+	for (int32 Index = 0; Index < InTimes.Num(); ++Index)
+	{
+		const FCatchupTimeRange& Zone = InTimes[Index];
+		
+		if (EnumHasAnyFlags(InFlags, EHoverCheckFlags::Icon))
+		{
+			const FGeometry Geometry = MakeCatchupZoneGeometry(Zone, InRangeToScreen, InAllottedGeometry);
+			if (Geometry.IsUnderLocation(InAbsoluteCursorPos))
+			{
+				return Index;
+			}
+		}
+		
+		if (EnumHasAnyFlags(InFlags, EHoverCheckFlags::Lines)
+			&& (IsLineHovered(Zone.StartTime, InAbsoluteCursorPos, InRangeToScreen, InAllottedGeometry)
+				|| IsLineHovered(GetCatchupRangeLastVerticalLineFrame(Zone), InAbsoluteCursorPos, InRangeToScreen, InAllottedGeometry)))
+		{
+			return Index;
+		}
+	}
+	return {};
+}
+
+static FAnalyzedHitchUIHoverInfo ComputeHoverState(
+	const FVector2f& InAbsoluteCursorPos, const FTimecodeHitchData& InData, const FScrubRangeToScreen& InRangeToScreen,
+	const FGeometry& InAllottedGeometry,
+	EHitchDrawFlags InFlags, EHoverCheckFlags InAreas
+	)
+{
+	using namespace Private;
+	
+	FAnalyzedHitchUIHoverInfo HoverInfo;
+	const FHitchVisualizationStyle Style;
+
+	if (EnumHasAnyFlags(InFlags, EHitchDrawFlags::DrawSkippedTimecodeMarkers))
+	{
+		HoverInfo.SkippedMarkerIndex = ComputeHoverState_Markers(
+			InData.SkippedTimecodeMarkers, InAbsoluteCursorPos, InRangeToScreen, InAllottedGeometry, Style.SkipMarkerStyle, InAreas
+			);
+	}
+	
+	if (EnumHasAnyFlags(InFlags, EHitchDrawFlags::DrawRepeatedTimecodeMarkers)
+		// Prefer hovering above items
+		&& !HoverInfo.IsHovered())
+	{
+		HoverInfo.RepeatedMarkerIndex = ComputeHoverState_Markers(
+			InData.RepeatedTimecodeMarkers, InAbsoluteCursorPos, InRangeToScreen, InAllottedGeometry, Style.RepeatMarkerStyle, InAreas
+			);
+	}
+	
+	if (EnumHasAnyFlags(InFlags, EHitchDrawFlags::DrawCatchupRanges)
+		// Prefer hovering above items
+		&& !HoverInfo.IsHovered())
+	{
+		HoverInfo.CatchupTimeIndex = ComputeHoverState_Catchup(
+			InData.CatchupTimes, InAbsoluteCursorPos, InRangeToScreen, InAllottedGeometry, InAreas
+			);
+	}
+
+	return HoverInfo;
+}
+}
+
+FAnalyzedHitchUIHoverInfo ComputeHoverStateForTimeSliderArea(
+	const FVector2f& InAbsoluteCursorPos,
+	const FTimecodeHitchData& InData, const FScrubRangeToScreen& InRangeToScreen,
+	const FGeometry& InAllottedGeometry, EHitchDrawFlags InFlags
+)
+{
+	const FAnalyzedHitchUIHoverInfo Info = ComputeHoverState(InAbsoluteCursorPos, InData, InRangeToScreen, InAllottedGeometry, InFlags,
+		Private::EHoverCheckFlags::All
+		);
+	return Info;
+}
+}
diff --git a/Engine/Plugins/VirtualProduction/Takes/Source/TakeRecorder/Private/Timecode/Visualization/Drawing/AnalyzedHitchUIDrawingUtils.h b/Engine/Plugins/VirtualProduction/Takes/Source/TakeRecorder/Private/Timecode/Visualization/Drawing/AnalyzedHitchUIDrawingUtils.h
new file mode 100644
--- /dev/null
+++ b/Engine/Plugins/VirtualProduction/Takes/Source/TakeRecorder/Private/Timecode/Visualization/Drawing/AnalyzedHitchUIDrawingUtils.h
@@ -0,0 +1,37 @@
+﻿// Copyright Epic Games, Inc. All Rights Reserved.
+
+#pragma once
+
+#include "EHitchDrawFlags.h"
+#include "HAL/Platform.h"
+#include "Math/Vector2D.h"
+
+class FSlateWindowElementList;
+namespace UE::TakeRecorder { struct FAnalyzedHitchUIHoverInfo; }
+namespace UE::TakeRecorder { struct FScrubRangeToScreen; }
+namespace UE::TakeRecorder { struct FTimecodeHitchData; }
+struct FGeometry;
+
+namespace UE::TakeRecorder::HitchVisualizationUI
+{
+/** Draws the top part (markers, etc.) of InData. */
+int32 DrawTimeSliderArea(
+	const FTimecodeHitchData& InData, const FAnalyzedHitchUIHoverInfo& InHoverInfo, const FScrubRangeToScreen& InRangeToScreen,
+	const FGeometry& InAllottedGeometry, FSlateWindowElementList& OutDrawElements, int32 InLayerId,
+	EHitchDrawFlags InFlags = EHitchDrawFlags::All
+	);
+	
+/** Draws the bottom part (markers, etc.) of InData. */
+int32 DrawTrackArea(
+	const FTimecodeHitchData& InData, const FAnalyzedHitchUIHoverInfo& InHoverInfo, const FScrubRangeToScreen& InRangeToScreen,
+	const FGeometry& InAllottedGeometry, FSlateWindowElementList& OutDrawElements, int32 InLayerId,
+	EHitchDrawFlags InFlags = EHitchDrawFlags::All
+	);
+	
+/** Determines the hovered geometry in the time slider area. */
+FAnalyzedHitchUIHoverInfo ComputeHoverStateForTimeSliderArea(
+	const FVector2f& InAbsoluteCursorPos,
+	const FTimecodeHitchData& InData, const FScrubRangeToScreen& InRangeToScreen,
+	const FGeometry& InAllottedGeometry, EHitchDrawFlags InFlags = EHitchDrawFlags::All
+	);
+}
diff --git a/Engine/Plugins/VirtualProduction/Takes/Source/TakeRecorder/Private/Timecode/Visualization/Drawing/AnalyzedHitchUIHoverInfo.h b/Engine/Plugins/VirtualProduction/Takes/Source/TakeRecorder/Private/Timecode/Visualization/Drawing/AnalyzedHitchUIHoverInfo.h
new file mode 100644
--- /dev/null
+++ b/Engine/Plugins/VirtualProduction/Takes/Source/TakeRecorder/Private/Timecode/Visualization/Drawing/AnalyzedHitchUIHoverInfo.h
@@ -0,0 +1,51 @@
+﻿// Copyright Epic Games, Inc. All Rights Reserved.
+
+#pragma once
+
+#include "Misc/Optional.h"
+#include "HAL/Platform.h"
+
+namespace UE::TakeRecorder
+{
+/**
+ * Information about which hitch visualization geometry is hovered.
+ * This is relevant only for when there is hitch analysis data.
+ * Only one of the optionals is set at any time.
+ */
+struct FAnalyzedHitchUIHoverInfo
+{
+	/** Index to the FTimecodeHitchData::SkippedTimecodeMarkers that should be drawn as set. */
+	TOptional<int32> SkippedMarkerIndex;
+	
+	/** Index to the FTimecodeHitchData::RepeatedTimecodeMarkers that should be drawn as set. */
+	TOptional<int32> RepeatedMarkerIndex;
+
+	/** Index to the FTimecodeHitchData::RepeatedTimecodeMarkers that should be drawn as hovered. */
+	TOptional<int32> CatchupTimeIndex;
+	
+	/** Resets this info so nothing is hovered. */
+	void Reset()
+	{
+		SkippedMarkerIndex.Reset();
+		RepeatedMarkerIndex.Reset();
+		CatchupTimeIndex.Reset();
+	}
+
+	/** @return Whether anything is hovered */
+	bool IsHovered() const { return SkippedMarkerIndex || RepeatedMarkerIndex || CatchupTimeIndex; }
+	/** @return Whether anything is hovered */
+	operator bool() const { return IsHovered(); }
+
+	friend bool operator==(const FAnalyzedHitchUIHoverInfo& Left, const FAnalyzedHitchUIHoverInfo& Right)
+	{
+		return Left.SkippedMarkerIndex == Right.SkippedMarkerIndex
+			&& Left.RepeatedMarkerIndex == Right.RepeatedMarkerIndex
+			&& Left.CatchupTimeIndex == Right.CatchupTimeIndex;
+	}
+
+	friend bool operator!=(const FAnalyzedHitchUIHoverInfo& Left, const FAnalyzedHitchUIHoverInfo& Right)
+	{
+		return !(Left == Right);
+	}
+};
+}
diff --git a/Engine/Plugins/VirtualProduction/Takes/Source/TakeRecorder/Private/Timecode/Visualization/Drawing/HitchHoverInfo.h b/Engine/Plugins/VirtualProduction/Takes/Source/TakeRecorder/Private/Timecode/Visualization/Drawing/HitchHoverInfo.h
deleted file mode 100644
diff --git a/Engine/Plugins/VirtualProduction/Takes/Source/TakeRecorder/Private/Timecode/Visualization/Drawing/MismatchedFrameRateDrawingUtils.cpp b/Engine/Plugins/VirtualProduction/Takes/Source/TakeRecorder/Private/Timecode/Visualization/Drawing/MismatchedFrameRateDrawingUtils.cpp
new file mode 100644
--- /dev/null
+++ b/Engine/Plugins/VirtualProduction/Takes/Source/TakeRecorder/Private/Timecode/Visualization/Drawing/MismatchedFrameRateDrawingUtils.cpp
@@ -0,0 +1,58 @@
+﻿// Copyright Epic Games, Inc. All Rights Reserved.
+
+#include "MismatchedFrameRateDrawingUtils.h"
+
+#include "ISequencer.h"
+#include "MismatchedFrameRateUIHoverInfo.h"
+#include "ScrubRangeToScreen.h"
+#include "TakeRecorderStyle.h"
+
+namespace UE::TakeRecorder::MismatchedFrameRateUI
+{
+static FGeometry MakeMarkerIconGeometry(
+	const FFrameNumber& InTime, const FScrubRangeToScreen& InRangeToScreen, const FGeometry& InAllottedGeometry, const FSlateBrush* InBrush
+)
+{
+	const float MarkerPos = InRangeToScreen.FrameToLocalX(InTime);
+	const FVector2f Position = FVector2f(MarkerPos, 0) + FTakeRecorderStyle::Get().GetVector("Hitching.MismatchedFrameRate.Offset");
+	const FVector2f Size = InBrush->GetImageSize();
+	return InAllottedGeometry.MakeChild(Size, FSlateLayoutTransform(Position));
+}
+	
+int32 DrawWarningIcon(
+	const ISequencer& InSequencer, const FMismatchedFrameRateUIHoverInfo& InHoverInfo, const FScrubRangeToScreen& InRangeToScreen,
+	const FGeometry& InAllottedGeometry, FSlateWindowElementList& OutDrawElements, int32 InLayerId
+)
+{
+	const FSlateBrush* Brush = FTakeRecorderStyle::Get().GetBrush("Hitching.MismatchedFrameRate.Icon");
+	const FLinearColor NormaTint = FTakeRecorderStyle::Get().GetColor("Hitching.MismatchedFrameRate.Normal");
+	const FLinearColor HoveredTint = FTakeRecorderStyle::Get().GetColor("Hitching.MismatchedFrameRate.Hovered");
+	
+	const FFrameNumber StartOfRange = InSequencer.GetPlaybackRange().GetLowerBoundValue();
+	const FGeometry IconGeometry = MakeMarkerIconGeometry(StartOfRange, InRangeToScreen, InAllottedGeometry, Brush);
+	
+	const bool bIsHovered = InHoverInfo.bIsWarningMarkerHovered;
+	FSlateDrawElement::MakeBox(
+		OutDrawElements,
+		InLayerId,
+		IconGeometry.ToPaintGeometry(),
+		Brush,
+		ESlateDrawEffect::None,
+		bIsHovered ? HoveredTint : NormaTint
+		);
+	
+	return InLayerId;
+}
+
+FMismatchedFrameRateUIHoverInfo ComputeHoverStateForTimeSliderArea(
+	const ISequencer& InSequencer, const FVector2f& InAbsoluteCursorPos,
+	const FScrubRangeToScreen& InRangeToScreen, const FGeometry& InAllottedGeometry
+	)
+{
+	const FSlateBrush* Brush = FTakeRecorderStyle::Get().GetBrush("Hitching.MismatchedFrameRate.Icon");
+	const FFrameNumber StartOfRange = InSequencer.GetPlaybackRange().GetLowerBoundValue();
+	const FGeometry IconGeometry = MakeMarkerIconGeometry(StartOfRange, InRangeToScreen, InAllottedGeometry, Brush);
+	
+	return FMismatchedFrameRateUIHoverInfo{ IconGeometry.IsUnderLocation(InAbsoluteCursorPos) };
+}
+}
diff --git a/Engine/Plugins/VirtualProduction/Takes/Source/TakeRecorder/Private/Timecode/Visualization/Drawing/MismatchedFrameRateDrawingUtils.h b/Engine/Plugins/VirtualProduction/Takes/Source/TakeRecorder/Private/Timecode/Visualization/Drawing/MismatchedFrameRateDrawingUtils.h
new file mode 100644
--- /dev/null
+++ b/Engine/Plugins/VirtualProduction/Takes/Source/TakeRecorder/Private/Timecode/Visualization/Drawing/MismatchedFrameRateDrawingUtils.h
@@ -0,0 +1,27 @@
+﻿// Copyright Epic Games, Inc. All Rights Reserved.
+
+#pragma once
+
+#include "HAL/Platform.h"
+#include "Math/Vector2D.h"
+
+class FSlateWindowElementList;
+class ISequencer;
+namespace UE::TakeRecorder { struct FMismatchedFrameRateUIHoverInfo; }
+namespace UE::TakeRecorder { struct FScrubRangeToScreen; }
+struct FGeometry;
+
+namespace UE::TakeRecorder::MismatchedFrameRateUI
+{
+/** Draws the icon for warning the user about mismatched frame rates. */
+int32 DrawWarningIcon(
+	const ISequencer& InSequencer, const FMismatchedFrameRateUIHoverInfo& InHoverInfo, const FScrubRangeToScreen& InRangeToScreen,
+	const FGeometry& InAllottedGeometry, FSlateWindowElementList& OutDrawElements, int32 InLayerId
+	);
+
+/** Determines the hovered geometry in the time slider area. */
+FMismatchedFrameRateUIHoverInfo ComputeHoverStateForTimeSliderArea(
+	const ISequencer& InSequencer, const FVector2f& InAbsoluteCursorPos,
+	const FScrubRangeToScreen& InRangeToScreen, const FGeometry& InAllottedGeometry
+	);
+}
diff --git a/Engine/Plugins/VirtualProduction/Takes/Source/TakeRecorder/Private/Timecode/Visualization/Drawing/MismatchedFrameRateUIHoverInfo.h b/Engine/Plugins/VirtualProduction/Takes/Source/TakeRecorder/Private/Timecode/Visualization/Drawing/MismatchedFrameRateUIHoverInfo.h
new file mode 100644
--- /dev/null
+++ b/Engine/Plugins/VirtualProduction/Takes/Source/TakeRecorder/Private/Timecode/Visualization/Drawing/MismatchedFrameRateUIHoverInfo.h
@@ -0,0 +1,27 @@
+﻿// Copyright Epic Games, Inc. All Rights Reserved.
+
+#pragma once
+
+#include "Misc/Optional.h"
+#include "HAL/Platform.h"
+
+namespace UE::TakeRecorder
+{
+/** This is hover information about the UI that is displayed when no analysis was performed due to mismatched frames. */
+struct FMismatchedFrameRateUIHoverInfo
+{
+	/** Whether the warning marker which tells the user that analysis was skipped is hovered. */
+	bool bIsWarningMarkerHovered = false;
+
+	/** Resets this info so nothing is hovered. */
+	void Reset()
+	{
+		bIsWarningMarkerHovered = false;
+	}
+
+	/** @return Whether anything is hovered */
+	bool IsHovered() const { return bIsWarningMarkerHovered; }
+	/** @return Whether anything is hovered */
+	operator bool() const { return IsHovered(); }
+};
+}
diff --git a/Engine/Plugins/VirtualProduction/Takes/Source/TakeRecorder/Private/Timecode/Visualization/Drawing/ScrubRangeToScreen.h b/Engine/Plugins/VirtualProduction/Takes/Source/TakeRecorder/Private/Timecode/Visualization/Drawing/ScrubRangeToScreen.h
--- a/Engine/Plugins/VirtualProduction/Takes/Source/TakeRecorder/Private/Timecode/Visualization/Drawing/ScrubRangeToScreen.h
+++ b/Engine/Plugins/VirtualProduction/Takes/Source/TakeRecorder/Private/Timecode/Visualization/Drawing/ScrubRangeToScreen.h
@@ -21,7 +21,13 @@ struct FScrubRangeToScreen : FSequencerTimeSliderController::FScrubRangeToScreen
 	/** FrameTime -> local Widget Space */
 	float FrameToLocalX(const FFrameTime& InFrameTime) const
 	{
-		const double Seconds = InFrameTime.FrameNumber / FrameRate; 
+		return FrameToLocalX(InFrameTime.FrameNumber);
+	}
+	
+	/** FrameTime -> local Widget Space */
+	float FrameToLocalX(const FFrameNumber& InFrameNumber) const
+	{
+		const double Seconds = InFrameNumber / FrameRate; 
 		return InputToLocalX(Seconds);
 	}
 };
diff --git a/Engine/Plugins/VirtualProduction/Takes/Source/TakeRecorder/Private/Timecode/Visualization/Drawing/VisualizationDrawingUtils.cpp b/Engine/Plugins/VirtualProduction/Takes/Source/TakeRecorder/Private/Timecode/Visualization/Drawing/VisualizationDrawingUtils.cpp
deleted file mode 100644
diff --git a/Engine/Plugins/VirtualProduction/Takes/Source/TakeRecorder/Private/Timecode/Visualization/Drawing/VisualizationDrawingUtils.h b/Engine/Plugins/VirtualProduction/Takes/Source/TakeRecorder/Private/Timecode/Visualization/Drawing/VisualizationDrawingUtils.h
deleted file mode 100644
diff --git a/Engine/Plugins/VirtualProduction/Takes/Source/TakeRecorder/Private/Timecode/Visualization/HitchAnalysis.cpp b/Engine/Plugins/VirtualProduction/Takes/Source/TakeRecorder/Private/Timecode/Visualization/HitchAnalysis.cpp
--- a/Engine/Plugins/VirtualProduction/Takes/Source/TakeRecorder/Private/Timecode/Visualization/HitchAnalysis.cpp
+++ b/Engine/Plugins/VirtualProduction/Takes/Source/TakeRecorder/Private/Timecode/Visualization/HitchAnalysis.cpp
@@ -182,23 +182,32 @@ static TArray<FCatchupTimeRange> AnalyzeCatchupZones(
 }
 }
 	
-TOptional<FTimecodeHitchData> AnalyseHitches(ISequencer& InSequencer)
+TValueOrError<FTimecodeHitchData, FHitchAnalysisErrorInfo> AnalyseHitches(ISequencer& InSequencer)
 {
 	UMovieSceneSequence* RootSequence = InSequencer.GetRootMovieSceneSequence();
 	UMovieScene* MovieScene = RootSequence->GetMovieScene();
 	UFrameHitchSceneDecoration* Decoration = MovieScene->FindDecoration<UFrameHitchSceneDecoration>();
 	if (!Decoration)
 	{
-		return {};
+		return MakeError(FHitchAnalysisErrorInfo{ EHitchAnalysisError::NoData });
+	}
+
+	// For now, analysis of mismatching frame rates is not supported. 
+	if (Decoration->RecordFrameRate != Decoration->TimecodeProviderFrameRate)
+	{
+		return MakeError(FHitchAnalysisErrorInfo{
+			EHitchAnalysisError::FrameRateMismatch,
+			FFrameRateMismatchData{ Decoration->RecordFrameRate, Decoration->TimecodeProviderFrameRate }
+		});
 	}
 	
 	const FFrameRate Rate = InSequencer.GetFocusedTickResolution();
-	Private::FAnalyzedMarkers Markers = Private::AnalyzeSkippedTimecodes(*Decoration, Rate, Decoration->FrameRate.Numerator);
-	return FTimecodeHitchData
+	Private::FAnalyzedMarkers Markers = Private::AnalyzeSkippedTimecodes(*Decoration, Rate, Decoration->RecordFrameRate.Numerator);
+	return MakeValue(FTimecodeHitchData
 	{
 		MoveTemp(Markers.Skipped), MoveTemp(Markers.Repeated),
 		Private::AnalyzeCatchupZones(*Decoration, Rate),
-	};
+	});
 }
 
 bool HasHitchData(ISequencer& InSequencer)
diff --git a/Engine/Plugins/VirtualProduction/Takes/Source/TakeRecorder/Private/Timecode/Visualization/HitchAnalysis.h b/Engine/Plugins/VirtualProduction/Takes/Source/TakeRecorder/Private/Timecode/Visualization/HitchAnalysis.h
--- a/Engine/Plugins/VirtualProduction/Takes/Source/TakeRecorder/Private/Timecode/Visualization/HitchAnalysis.h
+++ b/Engine/Plugins/VirtualProduction/Takes/Source/TakeRecorder/Private/Timecode/Visualization/HitchAnalysis.h
@@ -70,12 +70,51 @@ struct FTimecodeHitchData
 	}
 };
 
+enum class EHitchAnalysisError : uint8
+{
+	/** There was no UFrameHitchSceneDecoration attached, so we skipped analysis.  */
+	NoData,
+
+	/**
+	 * Take Recorder's record frame rate did not match the frame rate of underlying timecode provider.
+	 * 
+	 * For now, analysis is only supported when the frame rates are equal.
+	 * 
+	 * It's technically possible to implement but we just didn't implement it, yet. To implement, we need to change the timecode that is expected each
+	 * frame. For example, if the timecode provider is 48 FPS and TR 24 FPS, example frames we may expect would be
+	 * - frame 1 -> 17:42:42:05,
+	 * - frame 2 -> 17:42:42:07,
+	 * - frame 3 -> 17:42:42:09, etc.
+	 */
+	FrameRateMismatch,
+
+	Num
+};
+
+struct FFrameRateMismatchData
+{
+	/** Frame rate that TR recorded at. */
+	FFrameRate TakeRecorderRate;
+	/** Frame rate that the timecode provider had during recording. */
+	FFrameRate TimecodeRate;
+};
+	
+struct FHitchAnalysisErrorInfo
+{
+	/** The reason why analysis failed. */
+	EHitchAnalysisError Reason;
+
+	/** Set if Reason == FrameRateMismatch. */
+	TOptional<FFrameRateMismatchData> MismatchInfo;
+};
+
 /**
  * Tries to find the required tracks for analysing hitches, and then analyses them.
  * @return Analyzed hitch data if required tracks were found.
  */
-TOptional<FTimecodeHitchData> AnalyseHitches(ISequencer& InSequencer);
+TValueOrError<FTimecodeHitchData, FHitchAnalysisErrorInfo> AnalyseHitches(ISequencer& InSequencer);
 
 /** @return Whether there is any hitch data that can be analyzed. */
 bool HasHitchData(ISequencer& InSequencer);
+
 }
diff --git a/Engine/Plugins/VirtualProduction/Takes/Source/TakeRecorder/Private/Timecode/Visualization/HitchViewOptionUtils.cpp b/Engine/Plugins/VirtualProduction/Takes/Source/TakeRecorder/Private/Timecode/Visualization/HitchViewOptionUtils.cpp
--- a/Engine/Plugins/VirtualProduction/Takes/Source/TakeRecorder/Private/Timecode/Visualization/HitchViewOptionUtils.cpp
+++ b/Engine/Plugins/VirtualProduction/Takes/Source/TakeRecorder/Private/Timecode/Visualization/HitchViewOptionUtils.cpp
@@ -13,7 +13,7 @@
 #include "TakeRecorderCommands.h"
 #include "TakeRecorderHitchVisualizationSettings.h"
 #include "Templates/SharedPointer.h"
-#include "Timecode/Visualization/Sequencer/HitchVisualizationViewModel.h"
+#include "Timecode/Visualization/Sequencer/HitchViewModel_AnalyzedData.h"
 
 #define LOCTEXT_NAMESPACE "FHitchVisualizerInSequencer"
 
diff --git a/Engine/Plugins/VirtualProduction/Takes/Source/TakeRecorder/Private/Timecode/Visualization/HitchViewOptionUtils.h b/Engine/Plugins/VirtualProduction/Takes/Source/TakeRecorder/Private/Timecode/Visualization/HitchViewOptionUtils.h
--- a/Engine/Plugins/VirtualProduction/Takes/Source/TakeRecorder/Private/Timecode/Visualization/HitchViewOptionUtils.h
+++ b/Engine/Plugins/VirtualProduction/Takes/Source/TakeRecorder/Private/Timecode/Visualization/HitchViewOptionUtils.h
@@ -11,7 +11,7 @@ class FExtender;
 
 namespace UE::TakeRecorder
 {
-class FHitchVisualizationViewModel;
+class FHitchViewModel_AnalyzedData;
 	
 /**
  * Builds an extender for the Sequencer view options adding:
diff --git a/Engine/Plugins/VirtualProduction/Takes/Source/TakeRecorder/Private/Timecode/Visualization/Sequencer/HitchTooltipShower.cpp b/Engine/Plugins/VirtualProduction/Takes/Source/TakeRecorder/Private/Timecode/Visualization/Sequencer/HitchTooltipShower.cpp
deleted file mode 100644
diff --git a/Engine/Plugins/VirtualProduction/Takes/Source/TakeRecorder/Private/Timecode/Visualization/Sequencer/HitchTooltipShower.h b/Engine/Plugins/VirtualProduction/Takes/Source/TakeRecorder/Private/Timecode/Visualization/Sequencer/HitchTooltipShower.h
--- a/Engine/Plugins/VirtualProduction/Takes/Source/TakeRecorder/Private/Timecode/Visualization/Sequencer/HitchTooltipShower.h
+++ b/Engine/Plugins/VirtualProduction/Takes/Source/TakeRecorder/Private/Timecode/Visualization/Sequencer/HitchTooltipShower.h
@@ -4,10 +4,11 @@
 
 #if WITH_EDITOR
 
+#include "Framework/Application/SlateApplication.h"
 #include "Misc/Optional.h"
 #include "Templates/SharedPointer.h"
 #include "Templates/UnrealTemplate.h"
-#include "Timecode/Visualization/Drawing/HitchHoverInfo.h"
+#include "Widgets/SToolTip.h"
 
 class ISequencer;
 class IToolTip;
@@ -15,22 +16,33 @@ struct FFrameTime;
 
 namespace UE::TakeRecorder
 {
-struct FTimecodeHitchData;
+class FHitchViewModel;
 
-/** Responsible for generating a tooltip based on the currently hovered visualization UI. */
-class FHitchTooltipShower : public FNoncopyable
+/** Responsible for showing and hiding a tooltip based on the currently hovered visualization UI. */
+template<typename THoverState>
+class THitchTooltipShower : public FNoncopyable
 {
 public:
-
-	explicit FHitchTooltipShower(const TWeakPtr<ISequencer> InWeakSequencer)
-		: WeakSequencer(InWeakSequencer)
-	{}
 	
 	/** @return Whether tooltips are allowed to be shown currently. */
-	bool CanShowTooltips() const;
+	bool CanShowTooltips() const
+	{
+		FSlateApplication& SlateApplication = FSlateApplication::Get();
+		return SlateApplication.GetAllowTooltips() && SlateApplication.GetPressedMouseButtons().IsEmpty();
+	}
 	
 	/** Shows a visualization tooltip if currently hovering */
-	void UpdateTooltipState(const FTimecodeHitchData& HitchData, const FHitchHoverInfo& InHoverInfo);
+	void UpdateTooltipState(const THoverState& InHoverInfo, const TAttribute<FText>& InTooltipText)
+	{
+		if (InHoverInfo.IsHovered() && CanShowTooltips())
+		{
+			RecreateIfContentChanged(InHoverInfo, InTooltipText);
+		}
+		else
+		{
+			TooltipState.Reset();
+		}
+	}
 
 	/** Ensures the tooltip is hidden. */
 	void ClearTooltip()
@@ -39,27 +51,46 @@ public:
 	}
 
 private:
-
-	/** Used to convert frame times to display strings. */
-	const TWeakPtr<ISequencer> WeakSequencer;
 	
 	struct FTooltipState
 	{
 		/** The hovered UI elements when the tooltip was overriden. */
-		const FHitchHoverInfo HoverInfo;
+		const THoverState HoverInfo;
 
 		/** The hitch visualization tooltip we've temporarily spawned. */
 		const TSharedPtr<IToolTip> OurTooltip;
 
-		explicit FTooltipState(const FHitchHoverInfo& InHoverInfo, const FText& InTooltipTex);
-		~FTooltipState();
+		explicit FTooltipState(const THoverState& InHoverInfo, const TAttribute<FText>& InTooltipText)
+			: HoverInfo(InHoverInfo)
+			, OurTooltip(SNew(SToolTip)
+				.Text(InTooltipText)
+				.IsInteractive(true) // We need to set the tooltip as interactive as otherwise the tooltip is immediately hidden next frame via widget search
+				.OnSetInteractiveWindowLocation_Lambda([](FVector2D& OutPosition){ OutPosition =  FSlateApplication::Get().GetCursorPos(); })
+				)
+		{
+			FSlateApplication& SlateApplication = FSlateApplication::Get();
+			SlateApplication.SpawnToolTip(OurTooltip.ToSharedRef(), SlateApplication.GetCursorPos());
+		}
+		
+		~FTooltipState()
+		{
+			FSlateApplication& SlateApplication = FSlateApplication::Get();
+			SlateApplication.CloseToolTip();
+		}
 	};
 
 	/** Active tooltip override. Cleans up the tooltip on destruction. */
 	TOptional<FTooltipState> TooltipState;
 
 	/** (Re-)creates the tooltip if the hover a UI element is hovered. If the hovered UI has changed since last time, the tooltip is recreated. */
-	void RecreateIfContentChanged(const FTimecodeHitchData& HitchData, const FHitchHoverInfo& InHoverInfo, const ISequencer& InSequence);
+	void RecreateIfContentChanged(const THoverState& InHoverInfo, const TAttribute<FText>& InTooltipText)
+	{
+		const bool bNeedsTooltip = !TooltipState || InHoverInfo != TooltipState->HoverInfo;
+		if (bNeedsTooltip)
+		{
+			TooltipState.Emplace(InHoverInfo, InTooltipText);
+		}
+	}
 };
 }
 #endif
\ No newline at end of file
diff --git a/Engine/Plugins/VirtualProduction/Takes/Source/TakeRecorder/Private/Timecode/Visualization/Sequencer/HitchViewModel_AnalyzedData.cpp b/Engine/Plugins/VirtualProduction/Takes/Source/TakeRecorder/Private/Timecode/Visualization/Sequencer/HitchViewModel_AnalyzedData.cpp
new file mode 100644
--- /dev/null
+++ b/Engine/Plugins/VirtualProduction/Takes/Source/TakeRecorder/Private/Timecode/Visualization/Sequencer/HitchViewModel_AnalyzedData.cpp
@@ -0,0 +1,152 @@
+﻿// Copyright Epic Games, Inc. All Rights Reserved.
+
+#include "HitchViewModel_AnalyzedData.h"
+
+#include "Timecode/Visualization/Drawing/ScrubRangeToScreen.h"
+#include "Timecode/Visualization/Drawing/AnalyzedHitchUIDrawingUtils.h"
+#include "Timecode/Visualization/TakeRecorderHitchVisualizationSettings.h"
+#include "Timecode/Visualization/Drawing/MismatchedFrameRateDrawingUtils.h"
+#include "Widgets/SToolTip.h"
+
+#define LOCTEXT_NAMESPACE "FHitchVisualizationViewModel"
+
+namespace UE::TakeRecorder
+{
+namespace Private
+{
+static FText MakeHitchVisualizationTooltipText(
+	const ISequencer& InSequencer, const FTimecodeHitchData& InHitchData, const FAnalyzedHitchUIHoverInfo& InAnalysisHoverInfo
+)
+{
+	const TSharedPtr<INumericTypeInterface<double>> NumericTypeInterface = InSequencer.GetNumericTypeInterface();
+	const auto ConvertFrameToText = [&NumericTypeInterface](const FFrameNumber& InTime)
+	{
+		const FString AsFrameString = ensure(NumericTypeInterface)
+			? NumericTypeInterface->ToString(InTime.Value)
+			: FString::FromInt(InTime.Value);
+		return FText::FromString(AsFrameString);
+	};
+	
+	if (InAnalysisHoverInfo.RepeatedMarkerIndex && ensure(InHitchData.RepeatedTimecodeMarkers.IsValidIndex(*InAnalysisHoverInfo.RepeatedMarkerIndex)))
+	{
+		const FText Fmt = LOCTEXT(
+			"RepeatedMarkerFmt",
+			"Timecode was repeated at frame {0}.\nActual timecode:{1}\nExpected timecode:{2}"
+			);
+		const FUnexpectedTimecodeMarker& Data = InHitchData.RepeatedTimecodeMarkers[*InAnalysisHoverInfo.RepeatedMarkerIndex];
+		return FText::Format(Fmt,
+			ConvertFrameToText(Data.Frame),
+			FText::FromString(Data.ActualTimecode.ToString()),
+			FText::FromString(Data.ExpectedFrame.ToString())
+			);
+	}
+	
+	if (InAnalysisHoverInfo.SkippedMarkerIndex && ensure(InHitchData.SkippedTimecodeMarkers.IsValidIndex(*InAnalysisHoverInfo.SkippedMarkerIndex)))
+	{
+		const FText Fmt = LOCTEXT(
+			"SkippedMarkerFmt",
+			"A frame drop was detected at frame {0}.\nActual timecode:{1}\nExpected timecode:{2}\n\nData may be missing from record."
+			);
+		const FUnexpectedTimecodeMarker& Data = InHitchData.SkippedTimecodeMarkers[*InAnalysisHoverInfo.SkippedMarkerIndex];
+		return FText::Format(Fmt,
+			ConvertFrameToText(Data.Frame),
+			FText::FromString(Data.ActualTimecode.ToString()),
+			FText::FromString(Data.ExpectedFrame.ToString())
+			);
+	}
+	
+	if (InAnalysisHoverInfo.CatchupTimeIndex && ensure(InHitchData.CatchupTimes.IsValidIndex(*InAnalysisHoverInfo.CatchupTimeIndex)))
+	{
+		const FText Fmt = LOCTEXT(
+			"CatchupZoneFmt",
+			"Take Recorder was unable to keep up with the target timestep.\n\nData may have been lost or misplotted in this range [{0} - {1}]."
+			);
+		const FCatchupTimeRange& Data = InHitchData.CatchupTimes[*InAnalysisHoverInfo.CatchupTimeIndex];
+		return FText::Format(Fmt, ConvertFrameToText(Data.StartTime), ConvertFrameToText(Data.EndTime));
+	}
+	return FText::GetEmpty();
+}
+}
+	
+UE_SEQUENCER_DEFINE_CASTABLE(FHitchViewModel_AnalyzedData);
+
+int32 FHitchViewModel_AnalyzedData::OnPaintTimeSliderOverlay(
+	const FPaintArgs& Args, const FGeometry& AllottedGeometry, const FSlateRect& MyCullingRect, FSlateWindowElementList& OutDrawElements,
+	int32 LayerId, const FWidgetStyle& InWidgetStyle
+	)
+{
+	if (!CanPaint())
+	{
+		TooltipShowerHack.ClearTooltip();
+		return LayerId;
+	}
+
+	UTakeRecorderHitchVisualizationSettings* Settings = UTakeRecorderHitchVisualizationSettings::Get();
+	UpdateHoverState_TimeSlider(Args, AllottedGeometry, Settings->GetFlagsForUI());
+	
+	return HitchVisualizationUI::DrawTimeSliderArea(
+		ViewModelData, TimeSliderHoverInfo,
+		MakeScrubRangeToScreen(AllottedGeometry), AllottedGeometry, OutDrawElements, LayerId, Settings->GetFlagsForUI()
+		);
+}
+
+int32 FHitchViewModel_AnalyzedData::OnPaintTrackAreaOverlay(
+	const FPaintArgs& Args, const FGeometry& AllottedGeometry, const FSlateRect& MyCullingRect, FSlateWindowElementList& OutDrawElements,
+	int32 LayerId, const FWidgetStyle& InWidgetStyle
+	)
+{
+	if (!CanPaint())
+	{
+		return LayerId;
+	}
+
+	UTakeRecorderHitchVisualizationSettings* Settings = UTakeRecorderHitchVisualizationSettings::Get();
+	return HitchVisualizationUI::DrawTrackArea(
+		ViewModelData, TimeSliderHoverInfo,
+		MakeScrubRangeToScreen(AllottedGeometry), AllottedGeometry, OutDrawElements, LayerId, Settings->GetFlagsForUI()
+		);
+}
+
+	
+void FHitchViewModel_AnalyzedData::UpdateHoverState_TimeSlider(
+	const FPaintArgs& InArgs, const FGeometry& InAllottedGeometry, EHitchDrawFlags InFlags
+	)
+{
+	// InAllottedGeometry is in paint space. Mouse operations, etc. need to be done in tick space.
+	// The difference is that in paint space everything is shifted by GetWindowToDesktopTransform... so we need to undo that shift.
+	FGeometry TickSpaceGeometry = InAllottedGeometry;
+	TickSpaceGeometry.AppendTransform(FSlateLayoutTransform(InArgs.GetWindowToDesktopTransform()));
+	
+	// Don't draw the UI as hovered while the user is dragging the time slider.
+	const bool bCanHover = FSlateApplication::Get().GetPressedMouseButtons().IsEmpty();
+	if (bCanHover)
+	{
+		TimeSliderHoverInfo = HitchVisualizationUI::ComputeHoverStateForTimeSliderArea(
+			FSlateApplication::Get().GetCursorPos(), ViewModelData, MakeScrubRangeToScreen(TickSpaceGeometry),
+			TickSpaceGeometry, InFlags
+		);
+	}
+	else 
+	{
+		TimeSliderHoverInfo.Reset();
+		
+	}
+
+	// As we're manually managing the tooltip state, reevaluate whether it should be shown based off of our new hover state.
+	TooltipShowerHack.UpdateTooltipState(
+		TimeSliderHoverInfo,
+		TAttribute<FText>::CreateLambda([this]
+		{
+			const TSharedPtr<ISequencer> Sequencer = WeakOwningSequencer.Pin();
+			return Private::MakeHitchVisualizationTooltipText(*Sequencer, ViewModelData, TimeSliderHoverInfo);
+		}));
+}
+
+FScrubRangeToScreen FHitchViewModel_AnalyzedData::MakeScrubRangeToScreen(const FGeometry& AllottedGeometry) const
+{
+	const TSharedPtr<ISequencer> SequencerPin = WeakOwningSequencer.Pin();
+	return FScrubRangeToScreen(SequencerPin->GetFocusedTickResolution(), SequencerPin->GetViewRange(), AllottedGeometry.Size);
+}
+}
+
+#undef LOCTEXT_NAMESPACE
\ No newline at end of file
diff --git a/Engine/Plugins/VirtualProduction/Takes/Source/TakeRecorder/Private/Timecode/Visualization/Sequencer/HitchViewModel_AnalyzedData.h b/Engine/Plugins/VirtualProduction/Takes/Source/TakeRecorder/Private/Timecode/Visualization/Sequencer/HitchViewModel_AnalyzedData.h
new file mode 100644
--- /dev/null
+++ b/Engine/Plugins/VirtualProduction/Takes/Source/TakeRecorder/Private/Timecode/Visualization/Sequencer/HitchViewModel_AnalyzedData.h
@@ -0,0 +1,89 @@
+﻿// Copyright Epic Games, Inc. All Rights Reserved.
+
+#pragma once
+
+#if WITH_EDITOR
+#include "HitchTooltipShower.h"
+#include "MVVM/Extensions/ITopTimeSliderOverlayExtension.h"
+#include "MVVM/Extensions/ITrackAreaOverlayExtension.h"
+#include "MVVM/ViewModels/ViewModel.h"
+#include "Misc/Attribute.h"
+#include "Timecode/Visualization/HitchAnalysis.h"
+#include "Timecode/Visualization/TakeRecorderHitchVisualizationSettings.h"
+#include "Timecode/Visualization/Drawing/AnalyzedHitchUIHoverInfo.h"
+
+class IToolTip;
+
+namespace UE::TakeRecorder
+{
+struct FScrubRangeToScreen;
+	
+/**
+ * Draws analyzed hitch data:
+ * - Timecode skip & repeats draw a vertical line and a warning icon (in red ),
+ * - Catch up zones, when the engine was behind the target timecode (in yellow),
+ * - Draws a tooltip when UI in the timeslider section is hovered. The content is also drawn brighter. @see TooltipShowerHack.
+ *
+ * This is set up by FSequencerHitchVisualizer.
+ * @see FSequencerHitchVisualizer
+ */
+class FHitchViewModel_AnalyzedData
+	: public Sequencer::FViewModel
+	, public Sequencer::ITopTimeSliderOverlayExtension
+	, public Sequencer::ITrackAreaOverlayExtension
+{
+public:
+
+	UE_SEQUENCER_DECLARE_CASTABLE(FHitchViewModel_AnalyzedData, ITopTimeSliderOverlayExtension, ITrackAreaOverlayExtension);
+
+	explicit FHitchViewModel_AnalyzedData(
+		TWeakPtr<ISequencer> InWeakOwningSequencer, FTimecodeHitchData InHitchData, TAttribute<bool> InCanPaintAttr = true
+		)
+		: WeakOwningSequencer(MoveTemp(InWeakOwningSequencer))
+		, CanPaintAttr(MoveTemp(InCanPaintAttr))
+		, ViewModelData(MoveTemp(InHitchData))
+	{}
+	
+	//~ Begin ITimeSliderOverlayExtension Interface
+	virtual int32 OnPaintTimeSliderOverlay(const FPaintArgs& Args, const FGeometry& AllottedGeometry, const FSlateRect& MyCullingRect, FSlateWindowElementList& OutDrawElements, int32 LayerId, const FWidgetStyle& InWidgetStyle) override;
+	//~ End ITimeSliderOverlayExtension Interface
+	
+	//~ Begin ITrackAreaOverlayExtension Interface
+	virtual int32 OnPaintTrackAreaOverlay(const FPaintArgs& Args, const FGeometry& AllottedGeometry, const FSlateRect& MyCullingRect, FSlateWindowElementList& OutDrawElements, int32 LayerId, const FWidgetStyle& InWidgetStyle) override;
+	//~ End ITrackAreaOverlayExtension Interface
+
+private:
+
+	/** The sequencer that this model is saved in. */
+	const TWeakPtr<ISequencer> WeakOwningSequencer;
+
+	/** Decides whether we can paint. */
+	const TAttribute<bool> CanPaintAttr;
+
+	/** The hitch data that is being visualized. */
+	const FTimecodeHitchData ViewModelData;
+	/** In time slider area, info about which of HitchData should be drawn has hovered. Also used for determining what tooltip to show. */
+	FAnalyzedHitchUIHoverInfo TimeSliderHoverInfo;
+
+	/**
+	 * Handles showing and hiding the tooltip for our hitch visualization UI.
+	 *
+	 * This is a "hack" because we should prefer be using Sequencer's hotspot API: ideally, we'd tell the hotspot API what the user is currently
+	 * focusing. Then, the API would ask what the tooltip should be and handle showing it. However, the API is limited:
+	 * - the hotspot API is limited to track areas but we need it for the time slider area as well, and
+	 * - the hotspot API provides not way to show tooltips for now.
+	 * Once the hotspot API supports those operations, this hack should be replaced.
+	 */
+	THitchTooltipShower<FAnalyzedHitchUIHoverInfo> TooltipShowerHack;
+
+	/** Updates TimeSliderHoverInfo and the resulting tooltip state.*/
+	void UpdateHoverState_TimeSlider(const FPaintArgs& InArgs, const FGeometry& InAllottedGeometry, EHitchDrawFlags InFlags);
+
+	/** @return Utility for converting FFrameTime to widget paint positions. */
+	FScrubRangeToScreen MakeScrubRangeToScreen(const FGeometry& AllottedGeometry) const;
+
+	/** @return Whether this model should paint itself. */
+	bool CanPaint() const { return CanPaintAttr.Get(); }
+};
+}
+#endif
\ No newline at end of file
diff --git a/Engine/Plugins/VirtualProduction/Takes/Source/TakeRecorder/Private/Timecode/Visualization/Sequencer/HitchViewModel_MismatchedFrameRate.cpp b/Engine/Plugins/VirtualProduction/Takes/Source/TakeRecorder/Private/Timecode/Visualization/Sequencer/HitchViewModel_MismatchedFrameRate.cpp
new file mode 100644
--- /dev/null
+++ b/Engine/Plugins/VirtualProduction/Takes/Source/TakeRecorder/Private/Timecode/Visualization/Sequencer/HitchViewModel_MismatchedFrameRate.cpp
@@ -0,0 +1,91 @@
+﻿// Copyright Epic Games, Inc. All Rights Reserved.
+
+#include "HitchViewModel_MismatchedFrameRate.h"
+
+#include "Timecode/Visualization/Drawing/ScrubRangeToScreen.h"
+#include "Timecode/Visualization/Drawing/AnalyzedHitchUIDrawingUtils.h"
+#include "Timecode/Visualization/TakeRecorderHitchVisualizationSettings.h"
+#include "Timecode/Visualization/Drawing/MismatchedFrameRateDrawingUtils.h"
+#include "Widgets/SToolTip.h"
+
+#define LOCTEXT_NAMESPACE "FHitchVisualizationViewModel"
+
+namespace UE::TakeRecorder
+{
+namespace Private
+{
+static FText MakeHitchVisualizationTooltipText(const FFrameRateMismatchData& InHitchViewModel)
+{
+	const FText FormatTxt = LOCTEXT("SkippedAnalysis", 
+		"Hitch analysis was skipped.\n\nThe data was recorded at {0} but the timecode provider was at {1}."
+		"\nAnalysis of mismatched FPS is currently unsupported."
+		"\n\nTo prevent this, when recording, ensure that Take Recorder and your timecode provider have the same FPS."
+		);
+	return FText::Format(FormatTxt, InHitchViewModel.TakeRecorderRate.ToPrettyText(), InHitchViewModel.TimecodeRate.ToPrettyText());
+}
+}
+	
+UE_SEQUENCER_DEFINE_CASTABLE(FHitchViewModel_MismatchedFrameRate);
+
+int32 FHitchViewModel_MismatchedFrameRate::OnPaintTimeSliderOverlay(
+	const FPaintArgs& Args, const FGeometry& AllottedGeometry, const FSlateRect& MyCullingRect, FSlateWindowElementList& OutDrawElements,
+	int32 LayerId, const FWidgetStyle& InWidgetStyle
+	)
+{
+	if (!CanPaint())
+	{
+		TooltipShowerHack.ClearTooltip();
+		return LayerId;
+	}
+
+	UTakeRecorderHitchVisualizationSettings* Settings = UTakeRecorderHitchVisualizationSettings::Get();
+	UpdateHoverState_TimeSlider(Args, AllottedGeometry, Settings->GetFlagsForUI());
+	
+	const TSharedPtr<ISequencer> Sequencer = WeakOwningSequencer.Pin();
+	return ensure(Sequencer)
+		? MismatchedFrameRateUI::DrawWarningIcon(
+			*Sequencer, TimeSliderHoverInfo,
+			MakeScrubRangeToScreen(AllottedGeometry), AllottedGeometry, OutDrawElements, LayerId
+		)
+		: LayerId;
+}
+	
+void FHitchViewModel_MismatchedFrameRate::UpdateHoverState_TimeSlider(
+	const FPaintArgs& InArgs, const FGeometry& InAllottedGeometry, EHitchDrawFlags InFlags
+	)
+{
+	// InAllottedGeometry is in paint space. Mouse operations, etc. need to be done in tick space.
+	// The difference is that in paint space everything is shifted by GetWindowToDesktopTransform... so we need to undo that shift.
+	FGeometry TickSpaceGeometry = InAllottedGeometry;
+	TickSpaceGeometry.AppendTransform(FSlateLayoutTransform(InArgs.GetWindowToDesktopTransform()));
+	
+	// Don't draw the UI as hovered while the user is dragging the time slider.
+	const bool bCanHover = FSlateApplication::Get().GetPressedMouseButtons().IsEmpty();
+	if (!bCanHover)
+	{
+		TimeSliderHoverInfo.Reset();
+	}
+	else if (const TSharedPtr<ISequencer> Sequencer = WeakOwningSequencer.Pin())
+	{
+		TimeSliderHoverInfo = MismatchedFrameRateUI::ComputeHoverStateForTimeSliderArea(
+			*Sequencer, FSlateApplication::Get().GetCursorPos(), MakeScrubRangeToScreen(TickSpaceGeometry), TickSpaceGeometry
+		);
+	}
+
+	// As we're manually managing the tooltip state, reevaluate whether it should be shown based off of our new hover state.
+	TooltipShowerHack.UpdateTooltipState(
+		TimeSliderHoverInfo,
+		TAttribute<FText>::CreateLambda([this]
+		{
+			return Private::MakeHitchVisualizationTooltipText(ViewModelData);
+		}));
+}
+
+FScrubRangeToScreen FHitchViewModel_MismatchedFrameRate::MakeScrubRangeToScreen(const FGeometry& AllottedGeometry) const
+{
+	const TSharedPtr<ISequencer> SequencerPin = WeakOwningSequencer.Pin();
+	return FScrubRangeToScreen(SequencerPin->GetFocusedTickResolution(), SequencerPin->GetViewRange(), AllottedGeometry.Size);
+}
+}
+
+#undef LOCTEXT_NAMESPACE
\ No newline at end of file
diff --git a/Engine/Plugins/VirtualProduction/Takes/Source/TakeRecorder/Private/Timecode/Visualization/Sequencer/HitchViewModel_MismatchedFrameRate.h b/Engine/Plugins/VirtualProduction/Takes/Source/TakeRecorder/Private/Timecode/Visualization/Sequencer/HitchViewModel_MismatchedFrameRate.h
new file mode 100644
--- /dev/null
+++ b/Engine/Plugins/VirtualProduction/Takes/Source/TakeRecorder/Private/Timecode/Visualization/Sequencer/HitchViewModel_MismatchedFrameRate.h
@@ -0,0 +1,81 @@
+﻿// Copyright Epic Games, Inc. All Rights Reserved.
+
+#pragma once
+
+#if WITH_EDITOR
+#include "HitchTooltipShower.h"
+#include "MVVM/Extensions/ITopTimeSliderOverlayExtension.h"
+#include "MVVM/Extensions/ITrackAreaOverlayExtension.h"
+#include "MVVM/ViewModels/ViewModel.h"
+#include "Misc/Attribute.h"
+#include "Timecode/Visualization/Drawing/MismatchedFrameRateUIHoverInfo.h"
+#include "Timecode/Visualization/HitchAnalysis.h"
+#include "Timecode/Visualization/TakeRecorderHitchVisualizationSettings.h"
+
+class IToolTip;
+
+namespace UE::TakeRecorder
+{
+struct FScrubRangeToScreen;
+	
+/**
+ * Draws a warning icon in the time slider area.
+ *
+ * The warning tells the user that hitch analysis did not take place because the frame rate that Take Recorder recorded at was different from that
+ * frame rate that the underlying timecode provider sampled at.
+ */
+class FHitchViewModel_MismatchedFrameRate
+	: public Sequencer::FViewModel
+	, public Sequencer::ITopTimeSliderOverlayExtension
+{
+public:
+
+	UE_SEQUENCER_DECLARE_CASTABLE(FHitchViewModel_MismatchedFrameRate, ITopTimeSliderOverlayExtension);
+
+	explicit FHitchViewModel_MismatchedFrameRate(
+		TWeakPtr<ISequencer> InWeakOwningSequencer, FFrameRateMismatchData InHitchData, TAttribute<bool> InCanPaintAttr = true
+		)
+		: WeakOwningSequencer(MoveTemp(InWeakOwningSequencer))
+		, CanPaintAttr(MoveTemp(InCanPaintAttr))
+		, ViewModelData(MoveTemp(InHitchData))
+	{}
+	
+	//~ Begin ITimeSliderOverlayExtension Interface
+	virtual int32 OnPaintTimeSliderOverlay(const FPaintArgs& Args, const FGeometry& AllottedGeometry, const FSlateRect& MyCullingRect, FSlateWindowElementList& OutDrawElements, int32 LayerId, const FWidgetStyle& InWidgetStyle) override;
+	//~ End ITimeSliderOverlayExtension Interface
+
+private:
+
+	/** The sequencer that this model is saved in. */
+	const TWeakPtr<ISequencer> WeakOwningSequencer;
+
+	/** Decides whether we can paint. */
+	const TAttribute<bool> CanPaintAttr;
+
+	/** The hitch data that is being visualized. */
+	const FFrameRateMismatchData ViewModelData;
+	/** In time slider area, info about which of HitchData should be drawn has hovered. Also used for determining what tooltip to show. */
+	FMismatchedFrameRateUIHoverInfo TimeSliderHoverInfo;
+
+	/**
+	 * Handles showing and hiding the tooltip for our hitch visualization UI.
+	 *
+	 * This is a "hack" because we should prefer be using Sequencer's hotspot API: ideally, we'd tell the hotspot API what the user is currently
+	 * focusing. Then, the API would ask what the tooltip should be and handle showing it. However, the API is limited:
+	 * - the hotspot API is limited to track areas but we need it for the time slider area as well, and
+	 * - the hotspot API provides not way to show tooltips for now.
+	 * Once the hotspot API supports those operations, this hack should be replaced.
+	 */
+	THitchTooltipShower<FMismatchedFrameRateUIHoverInfo> TooltipShowerHack;
+
+	/** Updates TimeSliderHoverInfo and the resulting tooltip state.*/
+	void UpdateHoverState_TimeSlider(const FPaintArgs& InArgs, const FGeometry& InAllottedGeometry, EHitchDrawFlags InFlags);
+
+	/** @return Utility for converting FFrameTime to widget paint positions. */
+	FScrubRangeToScreen MakeScrubRangeToScreen(const FGeometry& AllottedGeometry) const;
+
+	/** @return Whether this model should paint itself. */
+	bool CanPaint() const { return CanPaintAttr.Get(); }
+};
+}
+#endif
\ No newline at end of file
diff --git a/Engine/Plugins/VirtualProduction/Takes/Source/TakeRecorder/Private/Timecode/Visualization/Sequencer/HitchVisualizationViewModel.cpp b/Engine/Plugins/VirtualProduction/Takes/Source/TakeRecorder/Private/Timecode/Visualization/Sequencer/HitchVisualizationViewModel.cpp
deleted file mode 100644
diff --git a/Engine/Plugins/VirtualProduction/Takes/Source/TakeRecorder/Private/Timecode/Visualization/Sequencer/HitchVisualizationViewModel.h b/Engine/Plugins/VirtualProduction/Takes/Source/TakeRecorder/Private/Timecode/Visualization/Sequencer/HitchVisualizationViewModel.h
deleted file mode 100644
diff --git a/Engine/Plugins/VirtualProduction/Takes/Source/TakeRecorder/Private/Timecode/Visualization/Sequencer/SequencerHitchVisualizer.cpp b/Engine/Plugins/VirtualProduction/Takes/Source/TakeRecorder/Private/Timecode/Visualization/Sequencer/SequencerHitchVisualizer.cpp
--- a/Engine/Plugins/VirtualProduction/Takes/Source/TakeRecorder/Private/Timecode/Visualization/Sequencer/SequencerHitchVisualizer.cpp
+++ b/Engine/Plugins/VirtualProduction/Takes/Source/TakeRecorder/Private/Timecode/Visualization/Sequencer/SequencerHitchVisualizer.cpp
@@ -3,7 +3,8 @@
 #if WITH_EDITOR
 #include "SequencerHitchVisualizer.h"
 
-#include "HitchVisualizationViewModel.h"
+#include "HitchViewModel_AnalyzedData.h"
+#include "HitchViewModel_MismatchedFrameRate.h"
 #include "ISequencerModule.h"
 #include "MVVM/ViewModels/SequencerEditorViewModel.h"
 #include "MovieScene.h"
@@ -58,8 +59,8 @@ void FSequencerHitchVisualizer::OnSequencerCreated(TSharedRef<ISequencer> InSequ
 {
 	using namespace UE::Sequencer;
 
-	TOptional<FTimecodeHitchData> AnalyzedData = AnalyseHitches(*InSequencer);
-	if (!AnalyzedData)
+	const TValueOrError<FTimecodeHitchData, FHitchAnalysisErrorInfo> AnalysisResult = AnalyseHitches(*InSequencer);
+	if (AnalysisResult.HasError() && AnalysisResult.GetError().Reason == EHitchAnalysisError::NoData)
 	{
 		return;
 	}
@@ -70,9 +71,16 @@ void FSequencerHitchVisualizer::OnSequencerCreated(TSharedRef<ISequencer> InSequ
 		const TSharedPtr<ISequencer> Sequencer = WeakSequencer.Pin();
 		return Sequencer && HasHitchData(*Sequencer);
 	});
-	const TSharedRef<FHitchVisualizationViewModel> Model = MakeShared<FHitchVisualizationViewModel>(InSequencer, *AnalyzedData, CanPaint);
-	InSequencer->GetViewModel()->GetEditorPanels().AddChild(Model);
-	CustomizedSequencers.Add(InSequencer);
+	
+	FViewModelChildren Panels = InSequencer->GetViewModel()->GetEditorPanels();
+	if (AnalysisResult.HasValue())
+	{
+		Panels.AddChild(MakeShared<FHitchViewModel_AnalyzedData>(InSequencer, AnalysisResult.GetValue(), CanPaint));
+	}
+	else
+	{
+		Panels.AddChild(MakeShared<FHitchViewModel_MismatchedFrameRate>(InSequencer, *AnalysisResult.GetError().MismatchInfo, CanPaint));
+	}
 
 	InSequencer->GetCommandBindings()->MapAction(
 		FTakeRecorderCommands::Get().ClearRecordingIntegrityData,
diff --git a/Engine/Plugins/VirtualProduction/Takes/Source/TakeRecorder/Private/Timecode/Visualization/Sequencer/SequencerHitchVisualizer.h b/Engine/Plugins/VirtualProduction/Takes/Source/TakeRecorder/Private/Timecode/Visualization/Sequencer/SequencerHitchVisualizer.h
--- a/Engine/Plugins/VirtualProduction/Takes/Source/TakeRecorder/Private/Timecode/Visualization/Sequencer/SequencerHitchVisualizer.h
+++ b/Engine/Plugins/VirtualProduction/Takes/Source/TakeRecorder/Private/Timecode/Visualization/Sequencer/SequencerHitchVisualizer.h
@@ -12,7 +12,7 @@ class ISequencer;
 
 namespace UE::TakeRecorder
 {
-class FHitchVisualizationViewModel;
+class FHitchViewModel_AnalyzedData;
 
 /** Hooks into Sequencer and extends the UI to visualize hitches. */
 class FSequencerHitchVisualizer : public FNoncopyable
