Ensure that we have a sequence asset before checking for a MovieScene object.

diff --git a/Engine/Plugins/VirtualProduction/Takes/Source/TakeRecorder/Private/Recorder/TakeRecorder.cpp b/Engine/Plugins/VirtualProduction/Takes/Source/TakeRecorder/Private/Recorder/TakeRecorder.cpp
--- a/Engine/Plugins/VirtualProduction/Takes/Source/TakeRecorder/Private/Recorder/TakeRecorder.cpp
+++ b/Engine/Plugins/VirtualProduction/Takes/Source/TakeRecorder/Private/Recorder/TakeRecorder.cpp
@@ -1109,7 +1109,7 @@ void UTakeRecorder::StopInternal(const bool bCancelled)
 		Sequencer->SetPlaybackStatus(EMovieScenePlayerStatus::Stopped);
 	}
 
-	UMovieScene* MovieScene = SequenceAsset->GetMovieScene();
+	UMovieScene* MovieScene = SequenceAsset ? SequenceAsset->GetMovieScene() : nullptr;
 	if (MovieScene)
 	{
 		MovieScene->SetClockSource(CachedClockSource);
