Take Recorder: STORY When FPS mismatch & hitch protection enabled, show warning icon next to Sequencer FPS.

diff --git a/Engine/Plugins/VirtualProduction/Takes/Source/TakeRecorder/Private/Timecode/HitchProtectionFrameRateModel.cpp b/Engine/Plugins/VirtualProduction/Takes/Source/TakeRecorder/Private/Timecode/HitchProtectionFrameRateModel.cpp
new file mode 100644
--- /dev/null
+++ b/Engine/Plugins/VirtualProduction/Takes/Source/TakeRecorder/Private/Timecode/HitchProtectionFrameRateModel.cpp
@@ -0,0 +1,53 @@
+﻿// Copyright Epic Games, Inc. All Rights Reserved.
+
+#include "HitchProtectionFrameRateModel.h"
+
+#include "Engine/Engine.h"
+#include "Engine/TimecodeProvider.h"
+#include "Internationalization/Internationalization.h"
+#include "Internationalization/Text.h"
+#include "Recorder/TakeRecorderSubsystem.h"
+#include "TakeRecorderSettings.h"
+
+#define LOCTEXT_NAMESPACE "HitchProtectionFrameRateModel"
+
+namespace UE::TakeRecorder::HitchProtectionFrameRateModel
+{
+struct FFrameRateInfo
+{
+	const FFrameRate TakeRecorderFrameRate = GEngine ? GEngine->GetEngineSubsystem<UTakeRecorderSubsystem>()->GetFrameRate() : FFrameRate(0, 1);
+	const TOptional<FFrameRate> TimecodeFrameRate = GEngine && GEngine->GetTimecodeProvider() ? GEngine->GetTimecodeProvider()->GetFrameRate() : TOptional<FFrameRate>();
+};
+	
+bool ShouldShowFrameRateWarning()
+{
+	const FFrameRateInfo FrameRateInfo;
+	return GetMutableDefault<UTakeRecorderProjectSettings>()->HitchProtectionSettings.bEnableHitchProtection
+		&& FrameRateInfo.TimecodeFrameRate
+		&& FrameRateInfo.TakeRecorderFrameRate != *FrameRateInfo.TimecodeFrameRate;
+}
+
+FText GetMismatchedFrameRateWarningTooltipText()
+{
+	if (!GEngine)
+	{
+		return FText::GetEmpty();
+	}
+
+	const FFrameRateInfo FrameRateInfo;
+	const FText TooltipFmt = LOCTEXT(
+		"MismatchedFPS",
+		"You have selected {0} but the timecode provider is set to {1}."
+		"\n\nThis is shown because you have enabled Hitch Protection. "
+		"\nHitch analysis is not supported for mismatched FPS."
+		"\nHitch protection will still run but there will be no visualization of hitches when you review the recording."
+		"\n\nTo fix this, ensure that Take Recorder and your timecode provider have the same FPS."
+		);
+	return FText::Format(TooltipFmt,
+		FrameRateInfo.TakeRecorderFrameRate.ToPrettyText(),
+		FrameRateInfo.TimecodeFrameRate.IsSet() ? FrameRateInfo.TimecodeFrameRate->ToPrettyText() : LOCTEXT("Nothing", "nothing")
+		);
+}
+}
+
+#undef LOCTEXT_NAMESPACE
\ No newline at end of file
diff --git a/Engine/Plugins/VirtualProduction/Takes/Source/TakeRecorder/Private/Timecode/HitchProtectionFrameRateModel.h b/Engine/Plugins/VirtualProduction/Takes/Source/TakeRecorder/Private/Timecode/HitchProtectionFrameRateModel.h
new file mode 100644
--- /dev/null
+++ b/Engine/Plugins/VirtualProduction/Takes/Source/TakeRecorder/Private/Timecode/HitchProtectionFrameRateModel.h
@@ -0,0 +1,20 @@
+﻿// Copyright Epic Games, Inc. All Rights Reserved.
+
+#pragma once
+
+class FText;
+
+/**
+ * This is the "ViewModel" for the warning icon displayed in the cockpit next to the FPS selection.
+ * 
+ * The warning icon warns the user to configure the same FPS for recording and timecode provider because hitch visualization relies
+ * on this (until UE-315784 is implemented).
+ */
+namespace UE::TakeRecorder::HitchProtectionFrameRateModel
+{
+/** @return Whether to show the warning icon. */
+bool ShouldShowFrameRateWarning();
+
+/** @return The tooltip text to show when the user hovers the warning icon. */
+FText GetMismatchedFrameRateWarningTooltipText();
+};
diff --git a/Engine/Plugins/VirtualProduction/Takes/Source/TakeRecorder/Private/Widgets/STakeRecorderCockpit.cpp b/Engine/Plugins/VirtualProduction/Takes/Source/TakeRecorder/Private/Widgets/STakeRecorderCockpit.cpp
--- a/Engine/Plugins/VirtualProduction/Takes/Source/TakeRecorder/Private/Widgets/STakeRecorderCockpit.cpp
+++ b/Engine/Plugins/VirtualProduction/Takes/Source/TakeRecorder/Private/Widgets/STakeRecorderCockpit.cpp
@@ -77,6 +77,8 @@
 #include "NamingTokensEngineSubsystem.h"
 #include "NamingTokensSpecifiers.h"
 #include "TakeRecorderNamingTokensData.h"
+#include "Timecode/HitchProtectionFrameRateModel.h"
+#include "Widgets/Layout/SScaleBox.h"
 
 extern UNREALED_API UEditorEngine* GEditor;
 
@@ -390,6 +392,15 @@ void STakeRecorderCockpit::Construct(const FArguments& InArgs)
 						SNew(SSpacer)
 					]
 
+					+ SHorizontalBox::Slot()
+					.AutoWidth()
+					.Padding(0, 0, 4, 0)
+					.HAlign(HAlign_Center)
+					.VAlign(VAlign_Center)
+					[
+						MakeHitchProtectionMismatchedFrameRateWarningIcon()
+					]
+
 					+ SHorizontalBox::Slot()
 					.AutoWidth()
 					[
@@ -667,6 +678,21 @@ void STakeRecorderCockpit::PostUndoRedo()
 	RefreshUndefinedTokens();
 }
 
+TSharedRef<SWidget> STakeRecorderCockpit::MakeHitchProtectionMismatchedFrameRateWarningIcon()
+{
+	return SNew(SScaleBox)
+	[
+		SNew(SImage)
+		.Image(FTakeRecorderStyle::Get().GetBrush("Hitching.MismatchedFrameRate.Icon"))
+		.ColorAndOpacity(FTakeRecorderStyle::Get().GetColor("Hitching.MismatchedFrameRate.Normal"))
+		.Visibility_Lambda([]
+		{
+			return UE::TakeRecorder::HitchProtectionFrameRateModel::ShouldShowFrameRateWarning() ? EVisibility::Visible : EVisibility::Hidden;
+		})
+		.ToolTipText_Lambda([]{ return UE::TakeRecorder::HitchProtectionFrameRateModel::GetMismatchedFrameRateWarningTooltipText(); })
+	];
+}
+
 void STakeRecorderCockpit::Refresh()
 {
 	if (const UTakeMetaData* TakeMetaData = GetMetaData())
diff --git a/Engine/Plugins/VirtualProduction/Takes/Source/TakeRecorder/Private/Widgets/STakeRecorderCockpit.h b/Engine/Plugins/VirtualProduction/Takes/Source/TakeRecorder/Private/Widgets/STakeRecorderCockpit.h
--- a/Engine/Plugins/VirtualProduction/Takes/Source/TakeRecorder/Private/Widgets/STakeRecorderCockpit.h
+++ b/Engine/Plugins/VirtualProduction/Takes/Source/TakeRecorder/Private/Widgets/STakeRecorderCockpit.h
@@ -79,6 +79,9 @@ private:
 	
 private:
 
+	/** @return When Hitch Protection is enabled, the widget that warns the user when recorded at FPS mismatches the timecode provider's FPS */
+	TSharedRef<SWidget> MakeHitchProtectionMismatchedFrameRateWarningIcon();
+
 	EVisibility GetTakeWarningVisibility() const;
 	FText GetTakeWarningText() const;
 
