Take Recorder: Add "Clear Recording Integrity Data" to time slider context menu when there is hitching data.

diff --git a/Engine/Plugins/VirtualProduction/Takes/Source/TakeRecorder/Private/Timecode/Visualization/HitchViewOptionUtils.cpp b/Engine/Plugins/VirtualProduction/Takes/Source/TakeRecorder/Private/Timecode/Visualization/HitchViewOptionUtils.cpp
--- a/Engine/Plugins/VirtualProduction/Takes/Source/TakeRecorder/Private/Timecode/Visualization/HitchViewOptionUtils.cpp
+++ b/Engine/Plugins/VirtualProduction/Takes/Source/TakeRecorder/Private/Timecode/Visualization/HitchViewOptionUtils.cpp
@@ -10,6 +10,7 @@
 #include "MovieScene.h"
 #include "MovieSceneSequence.h"
 #include "ScopedTransaction.h"
+#include "SequencerToolMenuContext.h"
 #include "TakeRecorderCommands.h"
 #include "TakeRecorderHitchVisualizationSettings.h"
 #include "Templates/SharedPointer.h"
@@ -19,13 +20,21 @@
 
 namespace UE::TakeRecorder
 {
+namespace MenuExtensionDetail
+{
+static FText GetRecordingIntegrityText()
+{
+	return LOCTEXT("RecordingIntegrity", "Recording Integrity");
+}
+}
+	
 TSharedRef<FExtender> ExtendSequencerViewOptionsWithHitchVisualization()
 {
 	const TSharedRef<FExtender> Extender = MakeShared<FExtender>();
 	Extender->AddMenuExtension(TEXT("MarkedFrames"), EExtensionHook::After, nullptr,
 		FMenuExtensionDelegate::CreateLambda([](FMenuBuilder& MenuBuilder)
 		{
-			MenuBuilder.BeginSection(TEXT("HitchVisualization"), LOCTEXT("RecordingIntegrity", "Recording Integrity"));
+			MenuBuilder.BeginSection(TEXT("HitchVisualization"), MenuExtensionDetail::GetRecordingIntegrityText());
 				
 			MenuBuilder.AddMenuEntry(
 				LOCTEXT("SkipMarkers.Title", "Show Frame Drop Markers"),
@@ -70,6 +79,28 @@ TSharedRef<FExtender> ExtendSequencerViewOptionsWithHitchVisualization()
 	return Extender;
 }
 
+void ExtendSequencerTimeSliderContextMenu()
+{
+	UToolMenu* ToolMenu = UToolMenus::Get()->ExtendMenu("Sequencer.TimelineMenu");
+
+	ToolMenu->AddDynamicSection("HitchVisualization", FNewToolMenuDelegate::CreateLambda([](UToolMenu* InMenu)
+	{
+		const USequencerToolMenuContext* SequencerContext = InMenu->FindContext<USequencerToolMenuContext>();
+		const TSharedPtr<ISequencer> SequencerPin = SequencerContext ? SequencerContext->WeakSequencer.Pin() : nullptr;
+		if (!ensure(SequencerContext))
+		{
+			return;
+		}
+
+		FToolMenuSection& Section = InMenu->AddSection(TEXT("HitchVisualization"), MenuExtensionDetail::GetRecordingIntegrityText());
+		if (HasHitchData(*SequencerPin))
+		{
+			// Just AddMenuEntry does not work. Seems like the command list is not added. FSequencerHitchVisualizer binds to the Sequencer command list.
+			Section.AddMenuEntryWithCommandList(FTakeRecorderCommands::Get().ClearRecordingIntegrityData, SequencerPin->GetCommandBindings());
+		}
+	}), FToolMenuInsert(NAME_None, EToolMenuInsertType::Last));
+}
+
 void Execute_ClearRecordingIntegrityData(TWeakPtr<ISequencer> InSequencer)
 {
 	const TSharedPtr<ISequencer> SequencerPin = InSequencer.Pin();
diff --git a/Engine/Plugins/VirtualProduction/Takes/Source/TakeRecorder/Private/Timecode/Visualization/HitchViewOptionUtils.h b/Engine/Plugins/VirtualProduction/Takes/Source/TakeRecorder/Private/Timecode/Visualization/HitchViewOptionUtils.h
--- a/Engine/Plugins/VirtualProduction/Takes/Source/TakeRecorder/Private/Timecode/Visualization/HitchViewOptionUtils.h
+++ b/Engine/Plugins/VirtualProduction/Takes/Source/TakeRecorder/Private/Timecode/Visualization/HitchViewOptionUtils.h
@@ -21,6 +21,8 @@ class FHitchViewModel_AnalyzedData;
  */
 TSharedRef<FExtender> ExtendSequencerViewOptionsWithHitchVisualization();
 
+void ExtendSequencerTimeSliderContextMenu();
+
 /** Handles the FTakeRecorderCommands::Get().ClearRecordingIntegrityData. */
 void Execute_ClearRecordingIntegrityData(TWeakPtr<ISequencer> InSequencer);
 /** Whether FTakeRecorderCommands::Get().ClearRecordingIntegrityData can be executed, i.e. movie scene is not read-only. */
diff --git a/Engine/Plugins/VirtualProduction/Takes/Source/TakeRecorder/Private/Timecode/Visualization/Sequencer/SequencerHitchVisualizer.cpp b/Engine/Plugins/VirtualProduction/Takes/Source/TakeRecorder/Private/Timecode/Visualization/Sequencer/SequencerHitchVisualizer.cpp
--- a/Engine/Plugins/VirtualProduction/Takes/Source/TakeRecorder/Private/Timecode/Visualization/Sequencer/SequencerHitchVisualizer.cpp
+++ b/Engine/Plugins/VirtualProduction/Takes/Source/TakeRecorder/Private/Timecode/Visualization/Sequencer/SequencerHitchVisualizer.cpp
@@ -42,6 +42,8 @@ FSequencerHitchVisualizer::FSequencerHitchVisualizer()
 		FOnSequencerCreated::FDelegate::CreateRaw(this, &FSequencerHitchVisualizer::OnSequencerCreated)
 		);
 	SequencerModule.GetViewMenuExtensibilityManager()->AddExtender(ViewOptionsExtender);
+
+	ExtendSequencerTimeSliderContextMenu();
 }
 
 FSequencerHitchVisualizer::~FSequencerHitchVisualizer()
