TakeRecorder: Make Take Recorder Source types public. Sources that weren't already exported now are. This is for consistency since all sources were previously exposed to blueprints.

diff --git a/Engine/Plugins/VirtualProduction/Takes/Source/TakeRecorderSources/Private/TakeRecorderActorSource.h b/Engine/Plugins/VirtualProduction/Takes/Source/TakeRecorderSources/Private/TakeRecorderActorSource.h
deleted file mode 100644
diff --git a/Engine/Plugins/VirtualProduction/Takes/Source/TakeRecorderSources/Private/TakeRecorderCameraCutSource.h b/Engine/Plugins/VirtualProduction/Takes/Source/TakeRecorderSources/Private/TakeRecorderCameraCutSource.h
deleted file mode 100644
diff --git a/Engine/Plugins/VirtualProduction/Takes/Source/TakeRecorderSources/Private/TakeRecorderLevelSequenceSource.h b/Engine/Plugins/VirtualProduction/Takes/Source/TakeRecorderSources/Private/TakeRecorderLevelSequenceSource.h
deleted file mode 100644
diff --git a/Engine/Plugins/VirtualProduction/Takes/Source/TakeRecorderSources/Private/TakeRecorderLevelVisibilitySource.h b/Engine/Plugins/VirtualProduction/Takes/Source/TakeRecorderSources/Private/TakeRecorderLevelVisibilitySource.h
deleted file mode 100644
diff --git a/Engine/Plugins/VirtualProduction/Takes/Source/TakeRecorderSources/Private/TakeRecorderMicrophoneAudioManager.h b/Engine/Plugins/VirtualProduction/Takes/Source/TakeRecorderSources/Private/TakeRecorderMicrophoneAudioManager.h
deleted file mode 100644
diff --git a/Engine/Plugins/VirtualProduction/Takes/Source/TakeRecorderSources/Private/TakeRecorderMicrophoneAudioSource.h b/Engine/Plugins/VirtualProduction/Takes/Source/TakeRecorderSources/Private/TakeRecorderMicrophoneAudioSource.h
deleted file mode 100644
diff --git a/Engine/Plugins/VirtualProduction/Takes/Source/TakeRecorderSources/Private/TakeRecorderNearbySpawnedActorSource.h b/Engine/Plugins/VirtualProduction/Takes/Source/TakeRecorderSources/Private/TakeRecorderNearbySpawnedActorSource.h
deleted file mode 100644
diff --git a/Engine/Plugins/VirtualProduction/Takes/Source/TakeRecorderSources/Private/TakeRecorderPlayerSource.h b/Engine/Plugins/VirtualProduction/Takes/Source/TakeRecorderSources/Private/TakeRecorderPlayerSource.h
deleted file mode 100644
diff --git a/Engine/Plugins/VirtualProduction/Takes/Source/TakeRecorderSources/Private/TakeRecorderWorldSource.h b/Engine/Plugins/VirtualProduction/Takes/Source/TakeRecorderSources/Private/TakeRecorderWorldSource.h
deleted file mode 100644
diff --git a/Engine/Plugins/VirtualProduction/Takes/Source/TakeRecorderSources/Public/TakeRecorderActorSource.h b/Engine/Plugins/VirtualProduction/Takes/Source/TakeRecorderSources/Public/TakeRecorderActorSource.h
new file mode 100644
--- /dev/null
+++ b/Engine/Plugins/VirtualProduction/Takes/Source/TakeRecorderSources/Public/TakeRecorderActorSource.h
@@ -0,0 +1,342 @@
+// Copyright Epic Games, Inc. All Rights Reserved.
+
+#pragma once
+
+#include "TakeRecorderSource.h"
+#include "TakeRecorderSourceProperty.h"
+#include "UObject/SoftObjectPtr.h"
+#include "TrackRecorders/IMovieSceneTrackRecorderHost.h"
+#include "Serializers/MovieSceneActorSerialization.h"
+#include "Serializers/MovieSceneManifestSerialization.h"
+#include "MovieSceneSequenceID.h"
+#include "TakeRecorderActorSource.generated.h"
+
+class UActorComponent;
+class UMovieSceneFolder;
+class UMovieScene;
+class UMovieSceneTrackRecorderSettings;
+class UMovieSceneAnimationTrackRecorderEditorSettings;
+
+#define UE_API TAKERECORDERSOURCES_API
+
+DECLARE_LOG_CATEGORY_EXTERN(ActorSerialization, Verbose, All);
+
+UENUM(BlueprintType)
+enum class ETakeRecorderActorRecordType : uint8
+{
+	Possessable,
+	Spawnable,
+	ProjectDefault
+};
+
+/**
+* This Take Recorder Source can record an actor from the World's properties.
+* Records the properties of the actor and the components on the actor and safely
+* handles new components being spawned at runtime and the actor being destroyed.
+*/
+UCLASS(MinimalAPI, Category="Actors")
+class UTakeRecorderActorSource : public UTakeRecorderSource, public IMovieSceneTrackRecorderHost
+{
+public:
+	GENERATED_BODY()
+
+	/** Reference to the actor in the world that should have it's properties recorded. */
+	UPROPERTY(EditAnywhere, BlueprintReadWrite, DisplayName="Source Actor", Category="Actor Source")
+	TSoftObjectPtr<AActor> Target;
+
+	/**
+	 * Should this actor be recorded as a Possessable in Sequencer? If so the resulting Object Binding	
+	 * will not create a Spawnable copy of this object and instead will possess this object in the level.
+	 */
+	UPROPERTY(EditAnywhere, BlueprintReadWrite, Category = "Actor Source")
+	ETakeRecorderActorRecordType RecordType;
+
+	/** 
+	 * Whether to ensure that the parent hierarchy is also recorded. If recording to possessable and the parent is not recorded, 
+	 * the recorded transforms will be in local space since the child will still be attached to the parent in the level after 
+	 * recording.  If recording to spawnable and the parent is not recorded, the recorded transforms will be in global space 
+	 * since the child will not be attached to the parent in the level.
+	 */
+	UPROPERTY(EditAnywhere, BlueprintReadWrite, Category = "Actor Source")
+	bool bRecordParentHierarchy;
+
+	/** Whether to perform key-reduction algorithms as part of the recording */
+	UPROPERTY(EditAnywhere, BlueprintReadWrite, Category = "Actor Source")
+	bool bReduceKeys;
+
+	/**
+	 * Lists the properties and components on the current actor and whether or not each property will be
+	 * recorded into a track in the resulting Level Sequence. 
+	 */
+	UPROPERTY(EditAnywhere, Instanced, BlueprintReadWrite, meta=(ShowInnerProperties), Category = "Actor Source")
+	TObjectPtr<UActorRecorderPropertyMap> RecordedProperties;
+
+	/** Include only the animation bones/curves that match this list */
+	UPROPERTY(EditAnywhere, BlueprintReadWrite, Category = "Actor Source")
+	TArray<FString> IncludeAnimationNames;
+
+	/** Exclude all animation bones/curves that match this list */
+	UPROPERTY(EditAnywhere, BlueprintReadWrite, Category = "Actor Source")
+	TArray<FString> ExcludeAnimationNames;
+
+	/** The level sequence that this source is being recorded into. Set during PreRecording, null after PostRecording. */
+	UPROPERTY()
+	TObjectPtr<ULevelSequence> TargetLevelSequence;
+
+	/** The root or uppermost level sequence that this source is being recorded into. Set during PreRecording, null after PostRecording. */
+	UPROPERTY()
+	TObjectPtr<ULevelSequence> RootLevelSequence;
+
+	/**
+	* Dynamically created list of settings objects for the different factories that are recording something 
+	* on this actor. If a Factory has no properties it can record the settings objects will not get created.
+	* Only one instance of this object exists for a factory and the factory recorder will be passed the shared 
+	* instance.
+	*/
+	UPROPERTY()
+	TArray<TObjectPtr<UObject>> FactorySettings;
+	
+	/**
+	* An array of section recorders created during the recording process that are capturing data about the actor/components.
+	* Will be an empty list when a recording is not in progress.
+	*/
+	UPROPERTY()
+	TArray<TObjectPtr<class UMovieSceneTrackRecorder>> TrackRecorders;
+
+	/** The parent actor source that generated this actor source (ie. through parenting or as an attached component). Null after PostRecording */
+	UPROPERTY()
+	TObjectPtr<UTakeRecorderActorSource> ParentSource;
+
+	/** Show Dialog during the (possibly) slow parts of the take recording */
+	UPROPERTY()
+	bool bShowProgressDialog;
+
+public:
+
+	/*
+	 * Add a take recorder source for the given actor. 
+	 *
+	 * @param InActor The actor to add a source for
+	 * @param InSources The sources to add the actor to
+	 * @return The added source or the source already present with the same actor
+	 */
+	UFUNCTION(BlueprintCallable, Category = "Take Recorder")
+	static UE_API UTakeRecorderSource* AddSourceForActor(AActor* InActor, UTakeRecorderSources* InSources);
+
+	/*
+	 * Remove the given actor from TakeRecorderSources.
+	 *
+	 * @param InActor The actor to remove from the sources
+	 * @param InSources The sources from where to remove the actor
+	 */
+	UFUNCTION(BlueprintCallable, Category = "Take Recorder")
+	static UE_API void RemoveActorFromSources(AActor* InActor, UTakeRecorderSources* InSources);
+
+	/*
+	 * External Sequencers can set whether spawnables are allowed for this Take Recorder 
+	 */
+	static bool AllowsSpawnableObjects();
+	static void SetAllowsSpawnableObjects(bool bInAllowsSpawnableObjects);
+
+public:
+	UE_API UTakeRecorderActorSource(const FObjectInitializer& ObjInit);
+	
+	// UTakeRecorderSource Interface
+	UE_API virtual bool IsValid() const override;
+	UE_API virtual TArray<UTakeRecorderSource*> PreRecording(ULevelSequence* InSequence, FMovieSceneSequenceID InSequenceID, ULevelSequence* InRootSequence, FManifestSerializer* InManifestSerializer) override;
+	UE_API virtual void StartRecording(const FTimecode& InSectionStartTimecode, const FFrameNumber& InSectionFirstFrame, class ULevelSequence* InSequence) override;
+	UE_API virtual void TickRecording(const FQualifiedFrameTime& CurrentSequenceTime) override;
+	UE_API virtual void StopRecording(class ULevelSequence* InSequence) override;
+	UE_API virtual TArray<UTakeRecorderSource*> PostRecording(class ULevelSequence* InSequence, class ULevelSequence* InRootSequence, const bool bCancelled) override;
+	UE_API virtual void FinalizeRecording() override;
+	virtual TArray<UObject*> GetAdditionalSettingsObjects() const { return TArray<UObject*>(FactorySettings); }
+	UE_API virtual FString GetSubsceneTrackName(ULevelSequence* InSequence) const override;
+	UE_API virtual FString GetSubsceneAssetName(ULevelSequence* InSequence) const override;
+	UE_API virtual void AddContentsToFolder(class UMovieSceneFolder* InFolder) override;
+	virtual FMovieSceneSequenceID GetSequenceID() const override { return TargetSequenceID; }
+
+	// ~UTakeRecorderSource Interface
+
+	/** Set the Target actor that we are going to record. Will reset the Recorded Property Map to defaults. */
+	UFUNCTION(BlueprintCallable, Category = "Take Recorder Actor Source")
+	UE_API void SetSourceActor(TSoftObjectPtr<AActor> InTarget);
+
+	UFUNCTION(BlueprintPure, Category = "Take Recorder Actor Source")
+	TSoftObjectPtr<AActor> GetSourceActor() const { return Target; }
+
+	/** Get the Guid of the Object Binding that this Actor Source created in the resulting Level Sequence. */
+	FGuid GetObjectBindingGuid() const
+	{
+		return CachedObjectBindingGuid;
+	}
+
+	/** Get the record type. If set to project default, gets the type from the project settings */
+	UE_API bool GetRecordToPossessable() const;
+
+protected:
+
+	/**
+	* Request that we re-build the map of which properties to record. This should only be called when the target Actor is changed as it will
+	* wipe out the users preference for which properties are going to get recorded and default them back to the Project Preferences for that
+	* particular Actor/Component.
+	*/
+	void RebuildRecordedPropertyMap();
+	void RebuildRecordedPropertyMapRecursive(const FFieldVariant& InObject, UActorRecorderPropertyMap* PropertyMap, const FString& OuterStructPath = FString());
+	
+	/**
+	* Looks at the given component and determines what the parent of this component is. For the root component and Actor Components the
+	* parent will be the root Property Map. For all other components, will attempt to find the Property Map which has a child Property Map
+	* that references the given component.
+	*
+	* Used to find which property map a new property map should be added to while respecting the component hierarchy.
+	*/
+	UActorRecorderPropertyMap* GetParentPropertyMapForComponent(UActorComponent* InComponent);
+	
+	/** Looks through the given Property Map recursively to find a property map which references the given component. */
+	UActorRecorderPropertyMap* GetPropertyMapForComponentRecursive(UActorComponent* InComponent, UActorRecorderPropertyMap* CurrentPropertyMap);
+
+	/**
+	* This is called when recording starts to generate the Section Recorders for the actor and all components that it currently has,
+	* as well as again during runtime for any newly added components.
+	*/
+	void CreateSectionRecordersRecursive(UObject* ObjectToRecord, UActorRecorderPropertyMap* PropertyMap, TArray<UObject*>& TraversedObjects);
+
+	/** Update our cached properties for what will be recorded. Done here so the UI doesn't have to iterate through map every frame. */
+	void UpdateCachedNumberOfRecordedProperties();
+
+	/** Returns the Guid of the Possessable in the specified sequence that represents the given actor, or an invalid Guid if the actor has no object binding in the sequence. */
+	FGuid ResolveActorFromSequence(AActor* InActor, ULevelSequence* CurrentSequence) const;
+
+	/** Remove the Possessable data for the given Guid from the sequence. Calls CleanExistingDataFromSequenceImpl afterwards for any other cleanup you may wish to do. */
+	void CleanExistingDataFromSequence(const FGuid& ForGuid, ULevelSequence& InSequence);
+
+	/** Called as part of PostRecording before Track Recorders are finalized. Calls PostProcessTrackRecordersImpl afterwards for any other post processing you wish to do before Track recorders are finalized. */
+	void PostProcessTrackRecorders(ULevelSequence* InSequence);
+
+	/** 
+	* Ensure that the Object Template this recording is recording into has the specified component. Used to initialize dynamically added components that don't exist in the CDO.
+	* @param InComponent - The component on the object that we want to check to see if it's a new component.
+	* @param OutComponent - The newly duplicated component for the CDO (if any)
+	* @return True if the given component was duplicated and OutComponent is valid.
+	*/
+	bool EnsureObjectTemplateHasComponent(UActorComponent* InComponent, UActorComponent*& OutComponent);
+
+	/** 
+	* Gets all components (both Scene and Actor) on the recorded Actor.
+	* @param OutArray - Set of components to add found components to.
+	* @param bUpdateReferencedActorList - If true will add any external Actors referenced to the NewReferencedActors list, and possibly print warnings.
+	*/
+	void GetAllComponents(TSet<TWeakObjectPtr<UActorComponent>>& OutArray, bool bUpdateReferencedActorList = false);
+
+	/**
+	* Gets all Scene components that are a child of the specified component. Use the Root Component if you want all
+	* child components on an actor.
+	* @param OnSceneComponent - The component who's children we look at.
+	* @param OutArray - Set of components to add found components to.
+	* @param bUpdateReferencedActorList - If true will add any external Actors referenced to the NewReferencedActors list, and possibly print warnings.
+	*/
+	void GetSceneComponents(USceneComponent* OnSceneComponent, TSet< TWeakObjectPtr<UActorComponent>>& OutArray, bool bUpdateReferencedActorList = false);
+
+	/**
+	* Returns the direct children of the specified scene component. Filters for ownership before returning.
+	*/
+	void GetChildSceneComponents(USceneComponent* OnSceneComponent, TSet< TWeakObjectPtr<UActorComponent>>& OutArray, bool bUpdateReferencedActorList = false);
+
+	/**
+	* Gets all Actor components on the recorded actor, not including Scene components.
+	*/
+	void GetActorComponents(AActor* OnActor, TSet< TWeakObjectPtr<UActorComponent>>& OutArray) const;
+
+	/**
+	* Iterates through the NewReferencedActors set and creates a new UTakeRecorderActorSource for each one. Only creates
+	* a new source if that actor is not already being recorded by another Actor Source, will not create an Actor source 
+	* for the actor we currently target.
+	*/
+	void CreateNewActorSourceForReferencedActors();
+	
+	/**
+	* Walks the hierarchy of the given Target actor to ensure all parents have been added to the NewReferencedActors set.
+	* This allows us to ensure that we can record transforms in local space as all parents will be recorded and we can rebuild
+	* the hierarchy through Attach Tracks in Sequencer.
+	*/
+	void EnsureParentHierarchyIsReferenced();
+
+	// IMovieSceneTrackRecorderHost Interface
+	/** Returns true if the other actor is also being recorded by the owning UTakeRecorderSources. Useful for checking if we're recording something we own is attached to. */
+	UE_API bool IsOtherActorBeingRecorded(AActor* OtherActor) const override;
+	/** Returns a valid guid if the other actor is also being recorded by the owning UTakeRecorderSources. Useful for knowing the guid of that Actor without knowing if it's a Possessable or a Spawnable. */
+	UE_API FGuid GetRecordedActorGuid(class AActor* OtherActor) const override;
+	/** Returns id of the active level sequence for the given Actor*/
+	UE_API FMovieSceneSequenceID GetLevelSequenceID(class AActor* OtherActor) override;
+	/** Returns generic track recorder settings */
+	UE_API FTrackRecorderSettings GetTrackRecorderSettings() const override;
+	/** Returns offset that may get set when recording a skeletal mesh animation. Needed to correctly transform attached children*/
+	UE_API FTransform GetRecordedActorAnimationInitialRootTransform(class AActor* OtherActor) const override;
+	/** Returns the root level sequence being recorded into */
+	ULevelSequence* GetRootLevelSequence() const override { return RootLevelSequence; };
+	// ~IMovieSceneTrackRecorderHost Interface
+	
+	/** Initializes an instance of the specified class if we don't already have it in our Settings array. */
+	void InitializeFactorySettingsObject(TSubclassOf<UMovieSceneTrackRecorderSettings> InClass);
+	/** Gets the already initialized instance of the specified class from the Settings array. */
+	UMovieSceneTrackRecorderSettings* GetSettingsObjectForFactory(TSubclassOf<UMovieSceneTrackRecorderSettings> InClass) const;
+
+	/** This is called after a Spawnable object template is created. Use this to modify any settings on the template object that need to be changed (ie: disabling auto-possession of pawns). */
+	UE_API virtual void PostProcessCreatedObjectTemplateImpl(AActor* ObjectTemplate);
+	virtual void CleanExistingDataFromSequenceImpl(const FGuid& ForGuid, ULevelSequence& InSequence) {}
+	virtual void PostProcessTrackRecordersImpl() {}
+	
+	
+public:
+	// UObject Interface
+	UE_API virtual void PostEditChangeProperty(struct FPropertyChangedEvent& PropertyChangedEvent) override;
+	UE_API virtual void PostDuplicate(bool bDuplicateForPIE) override;
+	UE_API virtual void PostEditUndo() override;
+	// ~UObject Interface
+private:
+
+	UE_API virtual const FSlateBrush* GetDisplayIconImpl() const override;
+	UE_API virtual FText GetCategoryTextImpl() const;
+	UE_API virtual FText GetDisplayTextImpl() const override;
+	UE_API virtual FText GetDescriptionTextImpl() const override;
+	UE_API virtual FText GetAddSourceDisplayTextImpl() const override;
+	
+	void ProcessRecordedTimes(ULevelSequence* InSequence);
+
+	/** Return either instance or global animation track recorder editor settings. */
+	UMovieSceneAnimationTrackRecorderEditorSettings* GetBestTrackRecorderEditorSettings() const;
+
+private:
+	/** Object Binding guid that is created in the Level Sequence when recording starts.*/
+	FGuid CachedObjectBindingGuid;
+
+	/** Array of actors that have some sort of referenced link to our actor (such as an object attached to our hierarchy) that need Actor Source recorders initialized for them. List is emptied after sources are created. */
+	TSet<class AActor*> NewReferencedActors;
+	/** Array of Actor Sources that we ended up creating that we need to clean up when we stop recording. */
+	TArray<class UTakeRecorderSource*> AddedActorSources;
+
+	/** 
+	* A pointer to the Object Template instance that was created inside the Target Level Sequence at the start of recording.
+	* Can be null (if recording to a possessable) or when we are not recording. 
+	*/
+	TWeakObjectPtr<class AActor> CachedObjectTemplate;
+	/**
+	* A set of components that we're currently recording. We compare the components from one frame to the next to see if any
+	* components have been added or removed so we can appropriately update their Spawn tracked. Empty unless a recording is
+	* in progress (but may still be empty if no components)
+	*/
+	TSet<TWeakObjectPtr<UActorComponent>> CachedComponentList;
+
+	/**
+	* ID of the TargetLevelSequence.
+	*/
+	FMovieSceneSequenceID TargetSequenceID;
+
+	/**
+	*  Serializer
+	*/
+	FActorSerializer ActorSerializer;
+
+};
+
+#undef UE_API
\ No newline at end of file
diff --git a/Engine/Plugins/VirtualProduction/Takes/Source/TakeRecorderSources/Public/TakeRecorderCameraCutSource.h b/Engine/Plugins/VirtualProduction/Takes/Source/TakeRecorderSources/Public/TakeRecorderCameraCutSource.h
new file mode 100644
--- /dev/null
+++ b/Engine/Plugins/VirtualProduction/Takes/Source/TakeRecorderSources/Public/TakeRecorderCameraCutSource.h
@@ -0,0 +1,66 @@
+// Copyright Epic Games, Inc. All Rights Reserved.
+
+#pragma once
+
+#include "TakeRecorderSource.h"
+#include "UObject/SoftObjectPtr.h"
+#include "Templates/SubclassOf.h"
+#include "GameFramework/Actor.h"
+#include "MovieSceneSequenceID.h"
+
+#include "TakeRecorderCameraCutSource.generated.h"
+
+class UWorld;
+class UTakeRecorderActorSource;
+
+#define UE_API TAKERECORDERSOURCES_API
+
+/** A recording source that detects camera switching and creates a camera cut track */
+UCLASS(MinimalAPI, Category="Other")
+class UTakeRecorderCameraCutSource : public UTakeRecorderSource
+{
+public:
+	GENERATED_BODY()
+
+	UE_API UTakeRecorderCameraCutSource(const FObjectInitializer& ObjInit);
+
+private:
+
+	// UTakeRecorderSource
+	UE_API virtual TArray<UTakeRecorderSource*> PreRecording(ULevelSequence* InSequence, FMovieSceneSequenceID InSequenceID, ULevelSequence* InRootSequence, FManifestSerializer* InManifestSerializer) override;
+	UE_API virtual void TickRecording(const FQualifiedFrameTime& CurrentTime) override;
+	UE_API virtual TArray<UTakeRecorderSource*> PostRecording(class ULevelSequence* InSequence, class ULevelSequence* InRootSequence, const bool bCancelled) override;
+	UE_API virtual FText GetDisplayTextImpl() const override;
+	UE_API virtual FText GetAddSourceDisplayTextImpl() const override;		
+	// This source does not support subscenes
+	virtual bool SupportsSubscenes() const override { return false; }
+
+private:
+
+	UPROPERTY()
+	TObjectPtr<UWorld> World;
+
+	/** The root or uppermost level sequence that this source is being recorded into. Set during PreRecording, null after PostRecording. */
+	UPROPERTY()
+	TObjectPtr<ULevelSequence> RootLevelSequence;
+
+
+	struct FCameraCutData
+	{
+		FCameraCutData(FGuid InGuid, FMovieSceneSequenceID InSequenceID, FQualifiedFrameTime InTime) 
+		: Guid(InGuid)
+		, SequenceID(InSequenceID)
+		, Time(InTime) {}
+
+		FGuid Guid;
+		FMovieSceneSequenceID SequenceID;
+		FQualifiedFrameTime Time;
+	};
+
+	TArray<FCameraCutData> CameraCutData;
+
+	/** Spawned actor sources to be removed at the end of recording */
+	TArray<TWeakObjectPtr<UTakeRecorderActorSource> > NewActorSources;
+};
+
+#undef UE_API
\ No newline at end of file
diff --git a/Engine/Plugins/VirtualProduction/Takes/Source/TakeRecorderSources/Public/TakeRecorderLevelSequenceSource.h b/Engine/Plugins/VirtualProduction/Takes/Source/TakeRecorderSources/Public/TakeRecorderLevelSequenceSource.h
new file mode 100644
--- /dev/null
+++ b/Engine/Plugins/VirtualProduction/Takes/Source/TakeRecorderSources/Public/TakeRecorderLevelSequenceSource.h
@@ -0,0 +1,44 @@
+// Copyright Epic Games, Inc. All Rights Reserved.
+
+#pragma once
+
+#include "TakeRecorderSource.h"
+#include "UObject/SoftObjectPtr.h"
+#include "TakeRecorderLevelSequenceSource.generated.h"
+
+class UTexture;
+class ALevelSequenceActor;
+class ULevelSequence;
+
+#define UE_API TAKERECORDERSOURCES_API
+
+/** Plays level sequence actors when recording starts */
+UCLASS(MinimalAPI,Category="Other")
+class UTakeRecorderLevelSequenceSource : public UTakeRecorderSource
+{
+public:
+	GENERATED_BODY()
+
+	UE_API UTakeRecorderLevelSequenceSource(const FObjectInitializer& ObjInit);
+
+	UPROPERTY(EditAnywhere, BlueprintReadWrite, Category = "Source")
+	TArray<TObjectPtr<ULevelSequence>> LevelSequencesToTrigger;
+
+private:
+
+	// UTakeRecorderSource
+	UE_API virtual TArray<UTakeRecorderSource*> PreRecording(ULevelSequence* InSequence, FMovieSceneSequenceID InSequenceID, ULevelSequence* InRootSequence, FManifestSerializer* InManifestSerializer) override;
+	UE_API virtual void StartRecording(const FTimecode& InSectionStartTimecode, const FFrameNumber& InSectionFirstFrame, class ULevelSequence* InSequence) override;
+	UE_API virtual void StopRecording(class ULevelSequence* InSequence) override;
+	UE_API virtual FText GetDisplayTextImpl() const override;
+	UE_API virtual FText GetDescriptionTextImpl() const override;
+	UE_API virtual FText GetAddSourceDisplayTextImpl() const override;	
+
+	// This source does not support subscenes since it's a playback source instead of a recording
+	virtual bool SupportsSubscenes() const override { return false; }
+
+	/** Transient level sequence actors to trigger, to be stopped and reset at the end of recording */
+	TArray<TWeakObjectPtr<ALevelSequenceActor>> ActorsToTrigger;
+};
+
+#undef UE_API
\ No newline at end of file
diff --git a/Engine/Plugins/VirtualProduction/Takes/Source/TakeRecorderSources/Public/TakeRecorderLevelVisibilitySource.h b/Engine/Plugins/VirtualProduction/Takes/Source/TakeRecorderSources/Public/TakeRecorderLevelVisibilitySource.h
new file mode 100644
--- /dev/null
+++ b/Engine/Plugins/VirtualProduction/Takes/Source/TakeRecorderSources/Public/TakeRecorderLevelVisibilitySource.h
@@ -0,0 +1,56 @@
+// Copyright Epic Games, Inc. All Rights Reserved.
+
+#pragma once
+
+#include "TakeRecorderSource.h"
+#include "TakeRecorderSources.h"
+#include "Templates/SubclassOf.h"
+#include "TakeRecorderLevelVisibilitySource.generated.h"
+
+#define UE_API TAKERECORDERSOURCES_API
+
+/** A recording source that records level visibilitiy */
+UCLASS(MinimalAPI, Abstract, config = EditorSettings, DisplayName = "Level Visibility Recorder Defaults")
+class UTakeRecorderLevelVisibilitySourceSettings : public UTakeRecorderSource
+{
+public:
+	GENERATED_BODY()
+
+	UE_API UTakeRecorderLevelVisibilitySourceSettings(const FObjectInitializer& ObjInit);
+
+	UE_API virtual void PostEditChangeProperty(FPropertyChangedEvent& PropertyChangedEvent) override;
+
+	// UTakeRecorderSource Interface
+	UE_API virtual FString GetSubsceneTrackName(ULevelSequence* InSequence) const override;
+	UE_API virtual FString GetSubsceneAssetName(ULevelSequence* InSequence) const override;
+	// ~UTakeRecorderSource Interface
+
+	/** Name of the recorded level visibility track name */
+	UPROPERTY(config, EditAnywhere, BlueprintReadWrite, Category = "Source")
+	FText LevelVisibilityTrackName;
+};
+
+/** A recording source that records level visibility state */
+UCLASS(MinimalAPI, Category="Other", config = EditorSettings)
+class UTakeRecorderLevelVisibilitySource : public UTakeRecorderLevelVisibilitySourceSettings
+{
+public:
+	GENERATED_BODY()
+
+	UE_API UTakeRecorderLevelVisibilitySource(const FObjectInitializer& ObjInit);
+
+private:
+
+	// UTakeRecorderSource
+	UE_API virtual TArray<UTakeRecorderSource*> PreRecording(ULevelSequence* InSequence, FMovieSceneSequenceID InSequenceID, ULevelSequence* InRootSequence, FManifestSerializer* InManifestSerializer) override;
+	UE_API virtual void TickRecording(const FQualifiedFrameTime& CurrentTime) override;
+	UE_API virtual void AddContentsToFolder(class UMovieSceneFolder* InFolder) override;
+	UE_API virtual FText GetDisplayTextImpl() const override;
+	UE_API virtual FText GetAddSourceDisplayTextImpl() const override;
+	UE_API virtual bool CanAddSource(UTakeRecorderSources* InSources) const override;
+
+private:
+	TWeakObjectPtr<class UMovieSceneLevelVisibilityTrack> CachedLevelVisibilityTrack;
+};
+
+#undef UE_API
\ No newline at end of file
diff --git a/Engine/Plugins/VirtualProduction/Takes/Source/TakeRecorderSources/Public/TakeRecorderMicrophoneAudioManager.h b/Engine/Plugins/VirtualProduction/Takes/Source/TakeRecorderSources/Public/TakeRecorderMicrophoneAudioManager.h
new file mode 100644
--- /dev/null
+++ b/Engine/Plugins/VirtualProduction/Takes/Source/TakeRecorderSources/Public/TakeRecorderMicrophoneAudioManager.h
@@ -0,0 +1,82 @@
+// Copyright Epic Games, Inc. All Rights Reserved.
+
+#pragma once
+
+#include "TakeRecorderSourceProperty.h"
+
+#include "TakeRecorderMicrophoneAudioManager.generated.h"
+
+#define UE_API TAKERECORDERSOURCES_API
+
+class IAudioCaptureEditor;
+class USoundWave;
+struct FTakeRecorderAudioSourceSettings;
+
+
+DECLARE_MULTICAST_DELEGATE_OneParam(FOnNotifySourcesOfDeviceChange, int);
+
+/** This class exposes the audio input device list via the project settings details. It does this in 
+*   conjunction with FAudioInputDevicePropertyCustomization. It also manages the IAudioCaptureEditor 
+*   object which handles the low level audio device recording.
+*/
+UCLASS(MinimalAPI, config=EditorPerProjectUserSettings, PerObjectConfig, DisplayName = "Audio Input Device")
+class UTakeRecorderMicrophoneAudioManager : public UTakeRecorderAudioInputSettings
+{
+public:
+	GENERATED_BODY()
+
+	UE_API UTakeRecorderMicrophoneAudioManager(const FObjectInitializer& ObjInit);
+
+	// Begin UObject Interface
+	UE_API virtual void PostEditChangeProperty(FPropertyChangedEvent& PropertyChangedEvent) override;
+	// End UObject Interface
+
+	/** Enumerates the audio devices present on the current machine */
+	UE_API virtual void EnumerateAudioDevices(bool InForceRefresh = false) override;
+
+	/** Returns input channel count for currently selected audio device */
+	UE_API virtual int32 GetDeviceChannelCount() override;
+
+	/** The audio device to use for this microphone source */
+	UPROPERTY(config, EditAnywhere, BlueprintReadWrite, Category = "Source", meta = (ShowOnlyInnerProperties))
+	FAudioInputDeviceProperty AudioInputDevice;
+
+	/** 
+	*  Calls StartRecording on the AudioRecorder object. This is called multiple times (once for each
+	*  microphone source, however, only the first call triggers the call to AudioRecorder.
+	*/
+	UE_API void StartRecording(int32 InChannelCount);
+	/**
+	 * Calls StopRecording on the AudioRecorder object. This is called multiple times (once for each
+	 * microphone source, however, only the first call triggers the call to AudioRecorder.
+	 */
+	UE_API void StopRecording();
+	/**
+	 *  Calls FinalizeRecording on the AudioRecorder object. This is called multiple times (once for each
+	 *  microphone source, however, only the first call triggers the call to AudioRecorder.
+	 */
+	UE_API void FinalizeRecording();
+
+
+	/** Fetches the USoundWave for this source after a Take has been recorded */
+	UE_API TObjectPtr<class USoundWave> GetRecordedSoundWave(const FTakeRecorderAudioSourceSettings& InSourceSettings);
+	/** Accessor for the OnNotifySourcesOfDeviceChange delegate list */
+	FOnNotifySourcesOfDeviceChange& GetOnNotifySourcesOfDeviceChange() { return OnNotifySourcesOfDeviceChange; }
+
+private:
+
+	/** Multicast delegate which notifies clients when the currently selected audio device changes. */
+	FOnNotifySourcesOfDeviceChange OnNotifySourcesOfDeviceChange;
+	/** Calls factory to create the AudioRecorder object */
+	UE_API TUniquePtr<IAudioCaptureEditor> CreateAudioRecorderObject();
+	/** Builds the list of audio devices which will be used in the device menu */
+	UE_API void BuildDeviceInfoArray();
+
+	/** Returns whether audio device with given Id was found during enumeration */
+	UE_API bool IsAudioDeviceAvailable(const FString& InDeviceId);
+
+	/** The audio recorder object which manages low level recording of audio data */
+	TUniquePtr<IAudioCaptureEditor> AudioRecorder;
+};
+
+#undef UE_API
diff --git a/Engine/Plugins/VirtualProduction/Takes/Source/TakeRecorderSources/Public/TakeRecorderMicrophoneAudioSource.h b/Engine/Plugins/VirtualProduction/Takes/Source/TakeRecorderSources/Public/TakeRecorderMicrophoneAudioSource.h
new file mode 100644
--- /dev/null
+++ b/Engine/Plugins/VirtualProduction/Takes/Source/TakeRecorderSources/Public/TakeRecorderMicrophoneAudioSource.h
@@ -0,0 +1,180 @@
+// Copyright Epic Games, Inc. All Rights Reserved.
+
+#pragma once
+
+#include "TakeRecorderSource.h"
+#include "TakeRecorderSources.h"
+#include "TakeRecorderSourceProperty.h"
+
+#include "GameFramework/Actor.h"
+#include "Templates/SubclassOf.h"
+#include "UObject/SoftObjectPtr.h"
+
+#include "TakeRecorderMicrophoneAudioSource.generated.h"
+
+#define UE_API TAKERECORDERSOURCES_API
+
+class UMovieSceneAudioTrack;
+class USoundWave;
+class UTakeRecorderMicrophoneAudioManager;
+struct FTakeRecorderAudioSourceSettings;
+
+/** A recording source that records microphone audio */
+UCLASS(MinimalAPI, Abstract, config=EditorSettings, DisplayName="Microphone Audio Recorder")
+class UTakeRecorderMicrophoneAudioSourceSettings : public UTakeRecorderSource
+{
+public:
+	GENERATED_BODY()
+
+	UE_API UTakeRecorderMicrophoneAudioSourceSettings(const FObjectInitializer& ObjInit);
+
+	UE_API virtual void PostEditChangeProperty(FPropertyChangedEvent& PropertyChangedEvent) override;
+	
+	// UTakeRecorderSource Interface
+	UE_API virtual FString GetSubsceneTrackName(ULevelSequence* InSequence) const override;
+	UE_API virtual FString GetSubsceneAssetName(ULevelSequence* InSequence) const override;
+	// ~UTakeRecorderSource Interface
+
+	/** Name of the audio source */
+	UPROPERTY(config, EditAnywhere, BlueprintReadWrite, Category = "Source")
+	FText AudioSourceName;
+
+	/** Name of the recorded audio track */
+	UPROPERTY(config, EditAnywhere, BlueprintReadWrite, Category = "Source", meta = (NamingTokens))
+	FText AudioTrackName;
+
+	/** The name of the audio asset.
+	 * Supports any of the following format specifiers that will be substituted when a take is recorded :
+	 * {day} - The day of the timestamp for the start of the recording.
+	 * {month} - The month of the timestamp for the start of the recording.
+	 * {year} - The year of the timestamp for the start of the recording.
+	 * {hour} - The hour of the timestamp for the start of the recording.
+	 * {minute} - The minute of the timestamp for the start of the recording.
+	 * {second} - The second of the timestamp for the start of the recording.
+	 * {take} - The take number.
+	 * {slate} - The slate string.
+	 */
+	UPROPERTY(config, EditAnywhere, BlueprintReadWrite, Category = "Source", meta = (NamingTokens))
+	FString AudioAssetName;
+
+	/** The name of the subdirectory audio will be placed in. Leave this empty to place into the same directory as the sequence base path 
+	 * Supports any of the following format specifiers that will be substituted when a take is recorded :
+	 * {day} - The day of the timestamp for the start of the recording.
+	 * {month} - The month of the timestamp for the start of the recording.
+	 * {year} - The year of the timestamp for the start of the recording.
+	 * {hour} - The hour of the timestamp for the start of the recording.
+	 * {minute} - The minute of the timestamp for the start of the recording.
+	 * {second} - The second of the timestamp for the start of the recording.
+	 * {take} - The take number.
+	 * {slate} - The slate string.
+	 */
+	UPROPERTY(config, EditAnywhere, BlueprintReadWrite, Category = "Source", meta = (NamingTokens))
+	FString AudioSubDirectory;
+};
+
+/** A recording source that records microphone audio */
+UCLASS(MinimalAPI, Category="Audio", config=EditorSettings)
+class UTakeRecorderMicrophoneAudioSource : public UTakeRecorderMicrophoneAudioSourceSettings
+{
+public:
+	GENERATED_BODY()
+
+	UE_API UTakeRecorderMicrophoneAudioSource(const FObjectInitializer& ObjInit);
+
+	// Begin UObject Interface
+	UE_API virtual void PostEditChangeProperty(FPropertyChangedEvent& PropertyChangedEvent) override;
+	// End UObject Interface
+
+	/** Gain in decibels to apply to recorded audio */
+	UPROPERTY(EditAnywhere, BlueprintReadWrite, Category = "Source", meta = (ClampMin = "0.0", UIMin = "0.0", UIMax = "20.0", ClampMax = "40.0"))
+	float AudioGain;
+
+	/** Whether or not to split mic channels into separate audio tracks. If not true, a max of 2 input channels is supported. */
+#if WITH_EDITORONLY_DATA
+	UPROPERTY(meta = (DeprecatedProperty, DeprecationMessage = "SplitAudioChannelsIntoSeparateTracks is deprecated."))
+	bool bSplitAudioChannelsIntoSeparateTracks_DEPRECATED;
+#endif
+
+	/** Replace existing recorded audio with any newly recorded audio */
+	UPROPERTY(EditAnywhere, BlueprintReadWrite, Category = "Source")
+	bool bReplaceRecordedAudio;
+
+	/** The audio device to use for this microphone source */
+	UPROPERTY(EditAnywhere, BlueprintReadWrite, Category = "Source", meta = (ShowOnlyInnerProperties))
+	FAudioInputDeviceChannelProperty AudioChannel;
+
+	// UTakeRecorderSource
+	UE_API virtual void Initialize() override;
+	UE_API virtual FString GetSubsceneTrackName(ULevelSequence* InSequence) const override;
+	UE_API virtual FString GetSubsceneAssetName(ULevelSequence* InSequence) const override;
+	UE_API virtual TArray<UTakeRecorderSource*> PreRecording(ULevelSequence* InSequence, FMovieSceneSequenceID InSequenceID, ULevelSequence* InRootSequence, FManifestSerializer* InManifestSerializer) override;
+	UE_API virtual void StartRecording(const FTimecode& InSectionStartTimecode, const FFrameNumber& InSectionFirstFrame, class ULevelSequence* InSequence) override;
+	UE_API virtual void TickRecording(const FQualifiedFrameTime& CurrentTime) override;
+	UE_API virtual void StopRecording(class ULevelSequence* InSequence) override;
+	UE_API virtual TArray<UTakeRecorderSource*> PostRecording(ULevelSequence* InSequence, class ULevelSequence* InRootSequence, const bool bCancelled) override;
+	UE_API virtual void FinalizeRecording() override;
+	UE_API virtual FText GetDisplayTextImpl() const override;
+	UE_API virtual FText GetAddSourceDisplayTextImpl() const override;
+
+	UE_API virtual void AddContentsToFolder(class UMovieSceneFolder* InFolder) override;
+	UE_API virtual bool CanAddSource(UTakeRecorderSources* InSources) const override;
+	// ~UTakeRecorderSource
+
+	/** Returns track display name, replacing any tokens if needed */
+	UE_API FString GetAudioTrackName(ULevelSequence* InSequence) const;
+	/** Returns the fully expanded asset name */
+	UE_API FString GetAudioAssetName(ULevelSequence* InSequence) const;
+
+	/** Sets the channel count supported by the currently selected audio device. */
+	UE_API void SetAudioDeviceChannelCount(int32 InChannelCount);
+	/** Delegate which receives notifications when the audio input device changes. */
+	UE_API void OnNotifySourcesOfDeviceChange(int32 InChannelCount);
+
+private:
+
+	/** Helper function for getting a pointer to the AudioInputManger object. */
+	static UE_API UTakeRecorderMicrophoneAudioManager* GetAudioInputManager();
+	
+	/** Parses recorded timecode values */
+	UE_API void ProcessRecordedTimes(ULevelSequence* InSequence);
+
+	/** Returns an array of booleans indicating which channel indexes are currently in use */
+	UE_API TUniquePtr<TArray<bool>> GetChannelsInUse(const int32 InDeviceChannelCount);
+	/** Fetches the USoundWave for this source after a Take has been recorded */
+	UE_API void GetRecordedSoundWave(ULevelSequence* InSequence);
+	/** Called when the user selects a new channel in the combobox for this source */
+	UE_API void SetCurrentInputChannel(int32 InChannelNumber);
+
+	/** Returns whether given track is being used by any other source */
+	UE_API bool IsTrackAssociatedWithAnySource(UMovieSceneAudioTrack* InAudioTrack);
+
+	/** Build asset name, appending channel name if needed for uniqueness */
+	UE_API FString GetUniqueAudioAssetName(ULevelSequence* InSequence) const;
+
+	/** Static, parameterized helper function for buidling unique asset names */	
+	static UE_API FString CreateUniqueAudioAssetName(ULevelSequence* InSequence, UTakeRecorderSources* InSources, const FString& InAssetName, const int32 InChannelNumber);
+
+	/** Evaluate our tokens and automatically create an appropriate context. */
+	UE_API bool TryEvaluateTokenString(const FString& InString, const ULevelSequence* InSequence, FString& OutString) const;
+
+private:
+	void DoAudioSetup();
+		
+	// Holds the Sequencer audio track for this source
+	TWeakObjectPtr<class UMovieSceneAudioTrack> CachedAudioTrack;
+
+	// Holds the USoundWave asset for a given take
+	TWeakObjectPtr<USoundWave> RecordedSoundWave;
+	// Caches the starting timecode for this take so it can be referenced when creating USoundWave assets
+	FTimecode StartTimecode;
+
+	// The user specified directory to store recorded audio assets in
+	FDirectoryPath AudioDirectory;
+	// The name of the UsoundWave asset
+	FString AssetFileName;
+
+	// Have we initialized our source.
+	bool bHaveDoneAudioSetup = false;
+};
+
+#undef UE_API
diff --git a/Engine/Plugins/VirtualProduction/Takes/Source/TakeRecorderSources/Public/TakeRecorderNearbySpawnedActorSource.h b/Engine/Plugins/VirtualProduction/Takes/Source/TakeRecorderSources/Public/TakeRecorderNearbySpawnedActorSource.h
new file mode 100644
--- /dev/null
+++ b/Engine/Plugins/VirtualProduction/Takes/Source/TakeRecorderSources/Public/TakeRecorderNearbySpawnedActorSource.h
@@ -0,0 +1,63 @@
+// Copyright Epic Games, Inc. All Rights Reserved.
+
+#pragma once
+
+#include "TakeRecorderSource.h"
+#include "UObject/SoftObjectPtr.h"
+#include "Templates/SubclassOf.h"
+#include "GameFramework/Actor.h"
+#include "TakeRecorderNearbySpawnedActorSource.generated.h"
+
+class UTexture;
+class UTakeRecorderActorSource;
+
+#define UE_API TAKERECORDERSOURCES_API
+
+/** A recording source that detects actors spawned close to the current camera, and captures them as spawnables */
+UCLASS(MinimalAPI, Category="Actors")
+class UTakeRecorderNearbySpawnedActorSource : public UTakeRecorderSource
+{
+public:
+	GENERATED_BODY()
+
+	UE_API UTakeRecorderNearbySpawnedActorSource(const FObjectInitializer& ObjInit);
+
+	/** The proximity to the current camera that an actor must be spawned in order to be recorded as a spawnable. If 0, proximity is disregarded. */
+	UPROPERTY(EditAnywhere, BlueprintReadWrite, DisplayName="Spawn Proximity", Category="Source", meta=(units=cm))
+	float Proximity;
+
+	/** Should we only record actors that pass the filter list?*/
+	UPROPERTY(EditAnywhere, BlueprintReadWrite, Category="Source")
+	bool bFilterSpawnedActors;
+	
+	/** A type filter to apply to spawned objects */
+	UPROPERTY(EditAnywhere, BlueprintReadWrite, Category="Source", meta=(EditCondition="bFilterSpawnedActors"))
+	TArray<TSubclassOf<AActor> > FilterTypes;
+
+private:
+
+	// UTakeRecorderSource
+	UE_API virtual TArray<UTakeRecorderSource*> PreRecording(ULevelSequence* InSequence, FMovieSceneSequenceID InSequenceID, ULevelSequence* InRootSequence, FManifestSerializer* InManifestSerializer) override;
+	UE_API virtual TArray<UTakeRecorderSource*> PostRecording(class ULevelSequence* InSequence, class ULevelSequence* InRootSequence, const bool bCancelled) override;
+	UE_API virtual FText GetDisplayTextImpl() const override;
+	UE_API virtual FText GetDescriptionTextImpl() const override;
+	UE_API virtual FText GetAddSourceDisplayTextImpl() const override;
+	// This source does not support subscenes (ie. "NearbySpawnedActors subscene"), but each of the spawned actors would be placed in subscenes if the option is enabled
+	virtual bool SupportsSubscenes() const override { return false; }
+
+	/** Handle actors being spawned */
+	void HandleActorSpawned(AActor* Actor, class ULevelSequence* InSequence);
+	
+	/** Is this actor valid for recording? Is it close enough? Is it a filtered type? */
+	bool IsActorValid(AActor* Actor);
+
+private:
+
+	/** Delegate handles for FOnActorSpawned events */
+	TMap<TWeakObjectPtr<UWorld>, FDelegateHandle> ActorSpawningDelegateHandles;
+
+	/** Spawned actor sources to be removed at the end of recording */
+	TArray<TWeakObjectPtr<UTakeRecorderActorSource> > SpawnedActorSources;
+};
+
+#undef UE_API
\ No newline at end of file
diff --git a/Engine/Plugins/VirtualProduction/Takes/Source/TakeRecorderSources/Public/TakeRecorderPlayerSource.h b/Engine/Plugins/VirtualProduction/Takes/Source/TakeRecorderSources/Public/TakeRecorderPlayerSource.h
new file mode 100644
--- /dev/null
+++ b/Engine/Plugins/VirtualProduction/Takes/Source/TakeRecorderSources/Public/TakeRecorderPlayerSource.h
@@ -0,0 +1,38 @@
+// Copyright Epic Games, Inc. All Rights Reserved.
+
+#pragma once
+
+#include "TakeRecorderSource.h"
+#include "TakeRecorderSources.h"
+#include "Templates/SubclassOf.h"
+#include "TakeRecorderPlayerSource.generated.h"
+
+class UTakeRecorderActorSource;
+
+#define UE_API TAKERECORDERSOURCES_API
+
+/** A recording source that records the current player */
+UCLASS(MinimalAPI, Category="Actors")
+class UTakeRecorderPlayerSource : public UTakeRecorderSource
+{
+public:
+	GENERATED_BODY()
+
+	UE_API UTakeRecorderPlayerSource(const FObjectInitializer& ObjInit);
+
+private:
+
+	// UTakeRecorderSource
+	UE_API virtual TArray<UTakeRecorderSource*> PreRecording(ULevelSequence* InSequence, FMovieSceneSequenceID InSequenceID, ULevelSequence* InRootSequence, FManifestSerializer* InManifestSerializer) override;
+	UE_API virtual TArray<UTakeRecorderSource*> PostRecording(class ULevelSequence* InSequence, class ULevelSequence* InRootSequence, const bool bCancelled) override;
+	UE_API virtual FText GetDisplayTextImpl() const override;
+	UE_API virtual FText GetAddSourceDisplayTextImpl() const override;				
+	UE_API virtual bool CanAddSource(UTakeRecorderSources* InSources) const override;
+
+	// This source does not support subscenes (ie. "Player subscene"), but the player would be placed in subscenes if the option is enabled
+	virtual bool SupportsSubscenes() const override { return false; }
+
+	TWeakObjectPtr<UTakeRecorderActorSource> PlayerActorSource;
+};
+
+#undef UE_API
\ No newline at end of file
diff --git a/Engine/Plugins/VirtualProduction/Takes/Source/TakeRecorderSources/Public/TakeRecorderWorldSource.h b/Engine/Plugins/VirtualProduction/Takes/Source/TakeRecorderSources/Public/TakeRecorderWorldSource.h
new file mode 100644
--- /dev/null
+++ b/Engine/Plugins/VirtualProduction/Takes/Source/TakeRecorderSources/Public/TakeRecorderWorldSource.h
@@ -0,0 +1,70 @@
+// Copyright Epic Games, Inc. All Rights Reserved.
+
+#pragma once
+
+#include "TakeRecorderSource.h"
+#include "TakeRecorderSources.h"
+#include "UObject/SoftObjectPtr.h"
+#include "Templates/SubclassOf.h"
+#include "GameFramework/Actor.h"
+#include "TakeRecorderWorldSource.generated.h"
+
+class UTexture;
+class UTakeRecorderActorSource;
+
+#define UE_API TAKERECORDERSOURCES_API
+
+/** A recording source that records world state */
+UCLASS(Abstract, MinimalAPI, config = EditorSettings, DisplayName = "World Recorder")
+class UTakeRecorderWorldSourceSettings : public UTakeRecorderSource
+{
+public:
+	GENERATED_BODY()
+
+	UE_API UTakeRecorderWorldSourceSettings(const FObjectInitializer& ObjInit);
+
+	UE_API virtual void PostEditChangeProperty(FPropertyChangedEvent& PropertyChangedEvent) override;
+
+	/** Record world settings */
+	UPROPERTY(config, EditAnywhere, BlueprintReadWrite, Category = "Source")
+	bool bRecordWorldSettings;
+
+	/** Add a binding and track for all actors that aren't explicitly being recorded */
+	UPROPERTY(config, EditAnywhere, BlueprintReadWrite, Category = "Source")
+	bool bAutotrackActors;
+};
+
+
+/** A recording source that records world state */
+UCLASS(MinimalAPI, Category="Actors")
+class UTakeRecorderWorldSource : public UTakeRecorderWorldSourceSettings
+{
+public:
+	GENERATED_BODY()
+
+	UE_API UTakeRecorderWorldSource(const FObjectInitializer& ObjInit);
+
+private:
+
+	// UTakeRecorderSource
+	UE_API virtual TArray<UTakeRecorderSource*> PreRecording(ULevelSequence* InSequence, FMovieSceneSequenceID InSequenceID, ULevelSequence* InRootSequence, FManifestSerializer* InManifestSerializer) override;
+	UE_API virtual TArray<UTakeRecorderSource*> PostRecording(class ULevelSequence* InSequence, class ULevelSequence* InRootSequence, const bool bCancelled) override;
+	UE_API virtual bool SupportsTakeNumber() const override { return false; }
+	UE_API virtual FText GetDisplayTextImpl() const override;
+	UE_API virtual FText GetAddSourceDisplayTextImpl() const override;		
+	UE_API virtual bool CanAddSource(UTakeRecorderSources* InSources) const override;
+
+	// This source does not support subscenes (ie. "World Settings subscene"), but the world settings actor would be placed in subscenes if the option is enabled
+	virtual bool SupportsSubscenes() const override { return false; }
+
+private:
+
+	/*
+	 * Autotrack actors in the world that aren't already being recorded
+	 */
+	void AutotrackActors(class ULevelSequence* InSequence, UWorld* InWorld);
+
+	TWeakObjectPtr<UTakeRecorderActorSource> WorldSource;
+};
+
+#undef UE_API
\ No newline at end of file
