Fix crash with checking cockpit errors

#rnx

diff --git a/Engine/Plugins/VirtualProduction/Takes/Source/TakeRecorder/Private/Widgets/STakeRecorderCockpit.cpp b/Engine/Plugins/VirtualProduction/Takes/Source/TakeRecorder/Private/Widgets/STakeRecorderCockpit.cpp
--- a/Engine/Plugins/VirtualProduction/Takes/Source/TakeRecorder/Private/Widgets/STakeRecorderCockpit.cpp
+++ b/Engine/Plugins/VirtualProduction/Takes/Source/TakeRecorder/Private/Widgets/STakeRecorderCockpit.cpp
@@ -537,13 +537,19 @@ void STakeRecorderCockpit::UpdateRecordError()
 		return;
 	}
 
-	UTakeRecorderSources*                  SourcesContainer = Sequence->FindMetaData<UTakeRecorderSources>();
-	TArrayView<UTakeRecorderSource* const> SourcesArray     = SourcesContainer ? SourcesContainer->GetSources() : TArrayView<UTakeRecorderSource* const>();
-	UTakeRecorderSource* const* Source = Algo::FindByPredicate(SourcesArray, [](const UTakeRecorderSource* Source)
+	UTakeRecorderSources* SourcesContainer = Sequence->FindMetaData<UTakeRecorderSources>();
+	bool bValidSources = false;
+	if (SourcesContainer)
 	{
-		return Source->bEnabled && Source->IsValid();
-	});
-	if (!Source)
+		TArrayView<UTakeRecorderSource* const> SourcesArray = SourcesContainer->GetSources();
+		UTakeRecorderSource* const* Source = Algo::FindByPredicate(SourcesArray, [](const UTakeRecorderSource* Source)
+		{
+			return Source->bEnabled && Source->IsValid();
+		});
+		bValidSources = Source != nullptr;
+	}
+
+	if (!bValidSources)
 	{
 		RecordErrorText = LOCTEXT("ErrorWidget_NoSources", "There are no currently enabled or valid sources to record from. Please add some above before recording.");
 		return;
