Turn on error reporting with save package.

diff --git a/Engine/Plugins/VirtualProduction/Takes/Source/TakesCore/Private/TakesUtils.cpp b/Engine/Plugins/VirtualProduction/Takes/Source/TakesCore/Private/TakesUtils.cpp
--- a/Engine/Plugins/VirtualProduction/Takes/Source/TakesCore/Private/TakesUtils.cpp
+++ b/Engine/Plugins/VirtualProduction/Takes/Source/TakesCore/Private/TakesUtils.cpp
@@ -22,6 +22,9 @@
 namespace TakesUtils
 {
 
+static bool bSaveTakesWithNoErrorFlag = false;
+static FAutoConsoleVariableRef  CVarSaveTakesWithNoError(TEXT("Takes.SaveTakesWithNoErrorFlag"), bSaveTakesWithNoErrorFlag, TEXT("For DrySquash only."));
+
 /** Helper function - get the first PIE world (or first PIE client world if there is more than one) */
 UWorld* GetFirstPIEWorld()
 {
@@ -114,11 +117,21 @@ void SaveAsset(UObject* InObject)
 
 	FSavePackageArgs SaveArgs;
 	SaveArgs.TopLevelFlags = RF_Standalone;
-	SaveArgs.SaveFlags = SAVE_NoError | SAVE_Async;
-	UPackage::SavePackage(Package, NULL, *PackageFileName, SaveArgs);
-
-	double ElapsedTime = FPlatformTime::Seconds() - StartTime;
-	UE_LOG(LogTakesCore, Log, TEXT("Saved %s in %0.2f seconds"), *PackageName, ElapsedTime);
+	SaveArgs.SaveFlags = SAVE_Async;
+	if (bSaveTakesWithNoErrorFlag)
+	{
+		SaveArgs.SaveFlags = SaveArgs.SaveFlags | SAVE_NoError;  
+	}
+	if (UPackage::SavePackage(Package, NULL, *PackageFileName, SaveArgs))
+	{
+		UPackage::WaitForAsyncFileWrites();
+		double ElapsedTime = FPlatformTime::Seconds() - StartTime;
+		UE_LOG(LogTakesCore, Log, TEXT("Saved %s in %0.2f seconds"), *PackageName, ElapsedTime);
+	}
+	else
+	{
+		UE_LOG(LogTakesCore, Log, TEXT("Failed to save %s."), *PackageName);
+	}
 }
 
 void CreateCameraCutTrack(ULevelSequence* LevelSequence, const FGuid& RecordedCameraGuid, const FMovieSceneSequenceID& SequenceID, const TRange<FFrameNumber>& InRange)
