TakeRecorder: Fix crash when creating level snapshot in multi-user session.

diff --git a/Engine/Plugins/VirtualProduction/Takes/Source/TakeRecorder/Private/TakeRecorderModule.cpp b/Engine/Plugins/VirtualProduction/Takes/Source/TakeRecorder/Private/TakeRecorderModule.cpp
--- a/Engine/Plugins/VirtualProduction/Takes/Source/TakeRecorder/Private/TakeRecorderModule.cpp
+++ b/Engine/Plugins/VirtualProduction/Takes/Source/TakeRecorder/Private/TakeRecorderModule.cpp
@@ -1,6 +1,8 @@
 // Copyright Epic Games, Inc. All Rights Reserved.
 
 #include "TakeRecorderModule.h"
+
+#include "TakesCoreLog.h"
 #include "TakeRecorderSettings.h"
 #include "SequencerUtilities.h"
 #include "Customization/TakeRecorderProjectSettingsCustomization.h"
@@ -388,6 +390,11 @@ void FTakeRecorderModule::ShutdownModule()
 
 UTakePreset* FTakeRecorderModule::GetPendingTake() const
 {
+	if (UE::IsSavingPackage(nullptr) || IsGarbageCollectingAndLockingUObjectHashTables())
+	{
+		UE_LOG(LogTakesCore, Verbose, TEXT("Cannot call FTakeRecorderModule::GetPendingTake while saving a package or garbage collecting."));
+		return nullptr;
+	}
 	return FindObject<UTakePreset>(nullptr, TEXT("/Temp/TakeRecorder/PendingTake.PendingTake"));
 }
 
