TakeRecorder: Fix GetLastRecordedLevelSequence not returning the last level sequence while reviewing a recording.

diff --git a/Engine/Plugins/VirtualProduction/Takes/Source/TakeRecorder/Private/Recorder/TakeRecorderSubsystemImplementation.cpp b/Engine/Plugins/VirtualProduction/Takes/Source/TakeRecorder/Private/Recorder/TakeRecorderSubsystemImplementation.cpp
--- a/Engine/Plugins/VirtualProduction/Takes/Source/TakeRecorder/Private/Recorder/TakeRecorderSubsystemImplementation.cpp
+++ b/Engine/Plugins/VirtualProduction/Takes/Source/TakeRecorder/Private/Recorder/TakeRecorderSubsystemImplementation.cpp
@@ -238,7 +238,7 @@ void UTakeRecorderSubsystemImplementation::SetRecordIntoLevelSequence(ULevelSequ
 
 bool UTakeRecorderSubsystemImplementation::CanReviewLastRecording() const
 {
-	if (GetLastRecordedLevelSequence() == nullptr)
+	if (GetSuppliedLevelSequence() != nullptr)
 	{
 		return false;
 	}
@@ -257,7 +257,6 @@ bool UTakeRecorderSubsystemImplementation::ReviewLastRecording()
 	if (LastRecordedLevelSequence.IsValid())
 	{
 		SuppliedLevelSequence = LastRecordedLevelSequence.Get();
-		LastRecordedLevelSequence = nullptr;
 		return true;
 	}
 
