TakeRecorder: Fix Take Recorder Subsystem IsReviewing() returning true for a frame when stopping a recording.

diff --git a/Engine/Plugins/VirtualProduction/Takes/Source/TakeRecorder/Private/Recorder/TakeRecorderSubsystemImplementation.cpp b/Engine/Plugins/VirtualProduction/Takes/Source/TakeRecorder/Private/Recorder/TakeRecorderSubsystemImplementation.cpp
--- a/Engine/Plugins/VirtualProduction/Takes/Source/TakeRecorder/Private/Recorder/TakeRecorderSubsystemImplementation.cpp
+++ b/Engine/Plugins/VirtualProduction/Takes/Source/TakeRecorder/Private/Recorder/TakeRecorderSubsystemImplementation.cpp
@@ -709,7 +709,9 @@ void UTakeRecorderSubsystemImplementation::ImportPreset(const FAssetData& InPres
 
 bool UTakeRecorderSubsystemImplementation::IsReviewing() const
 {
-	return !IsRecording() && (TakeMetaData && TakeMetaData->Recorded() && GetTakeRecorderMode() != ETakeRecorderMode::RecordIntoSequence);
+	return !IsRecording() && SuppliedLevelSequence &&
+		(TakeMetaData && TakeMetaData->Recorded()
+			&& GetTakeRecorderMode() != ETakeRecorderMode::RecordIntoSequence);
 }
 
 bool UTakeRecorderSubsystemImplementation::IsRecording() const
