Check for null source while iterating over source list.

#rnx

diff --git a/Engine/Plugins/VirtualProduction/Takes/Source/TakeRecorder/Private/Widgets/STakeRecorderCockpit.cpp b/Engine/Plugins/VirtualProduction/Takes/Source/TakeRecorder/Private/Widgets/STakeRecorderCockpit.cpp
--- a/Engine/Plugins/VirtualProduction/Takes/Source/TakeRecorder/Private/Widgets/STakeRecorderCockpit.cpp
+++ b/Engine/Plugins/VirtualProduction/Takes/Source/TakeRecorder/Private/Widgets/STakeRecorderCockpit.cpp
@@ -544,7 +544,7 @@ void STakeRecorderCockpit::UpdateRecordError()
 		TArrayView<UTakeRecorderSource* const> SourcesArray = SourcesContainer->GetSources();
 		UTakeRecorderSource* const* Source = Algo::FindByPredicate(SourcesArray, [](const UTakeRecorderSource* Source)
 		{
-			return Source->bEnabled && Source->IsValid();
+			return Source && Source->bEnabled && Source->IsValid();
 		});
 		bValidSources = Source != nullptr;
 	}
