Check for external references before save.

#rnx

diff --git a/Engine/Plugins/VirtualProduction/Takes/Source/TakeRecorder/Private/Recorder/TakeRecorder.cpp b/Engine/Plugins/VirtualProduction/Takes/Source/TakeRecorder/Private/Recorder/TakeRecorder.cpp
--- a/Engine/Plugins/VirtualProduction/Takes/Source/TakeRecorder/Private/Recorder/TakeRecorder.cpp
+++ b/Engine/Plugins/VirtualProduction/Takes/Source/TakeRecorder/Private/Recorder/TakeRecorder.cpp
@@ -1217,6 +1217,9 @@ void UTakeRecorder::StopInternal(const bool bCancelled)
 			{
 				if (Sequencer)
 				{
+					// Ensure we don't have any external references prior to saving.
+					TakesUtils::CleanupExternalReferences(SequenceAsset);
+					
 					// Sequencer has a feature for scrubbing to a particular frame and capturing it as thumbnail.
 					// In 5.5, this was added but hidden behind a CVar for stability reason.
 					// In BaseTakes.ini, by default, we configure our Sequencer instance to capture the MiddleFrame.
diff --git a/Engine/Plugins/VirtualProduction/Takes/Source/TakesCore/Private/TakesUtils.cpp b/Engine/Plugins/VirtualProduction/Takes/Source/TakesCore/Private/TakesUtils.cpp
--- a/Engine/Plugins/VirtualProduction/Takes/Source/TakesCore/Private/TakesUtils.cpp
+++ b/Engine/Plugins/VirtualProduction/Takes/Source/TakesCore/Private/TakesUtils.cpp
@@ -17,6 +17,7 @@
 #include "Sections/MovieSceneCameraCutSection.h"
 #include "UObject/SavePackage.h"
 
+#include "Serialization/ArchiveReplaceOrClearExternalReferences.h"
 #define LOCTEXT_NAMESPACE "TakesUtils"
 
 namespace TakesUtils
@@ -25,6 +26,52 @@ namespace TakesUtils
 static bool bSaveTakesWithNoErrorFlag = false;
 static FAutoConsoleVariableRef  CVarSaveTakesWithNoError(TEXT("Takes.SaveTakesWithNoErrorFlag"), bSaveTakesWithNoErrorFlag, TEXT("For DrySquash only."));
 
+class FArchiveClearExternalReferences : public FArchiveReplaceOrClearExternalReferences<UObject>
+{
+	using Super = FArchiveReplaceOrClearExternalReferences<UObject>;
+
+	TArray<UObject*> RemovedObjects;
+public:
+	
+	FArchiveClearExternalReferences(UObject* InSearchObject, const TMap<UObject*, UObject*>& InReplacementMap, UPackage* InDestPackage)
+		: Super(InSearchObject, InReplacementMap, InDestPackage, EArchiveReplaceObjectFlags::NullPrivateRefs | EArchiveReplaceObjectFlags::DelayStart )
+	{
+		SerializeSearchObject();
+	}
+
+	const TArray<UObject*>& GetRemovedObjects() const { return RemovedObjects; }
+
+	virtual FArchive& operator<<(UObject*& Obj) override
+	{
+		UObject* Original = Obj;
+		Super::operator<<(Obj);
+
+		if (Original && !Obj)
+		{
+			RemovedObjects.AddUnique(Original);
+		}
+		
+		return *this;
+	}
+};
+
+static void LogRemovedExternalObjects(UObject* ReferencingObject, const TArray<UObject*>& InRemovedObjects)
+{
+	if (InRemovedObjects.IsEmpty())
+	{
+		return;
+	}
+
+	const FString List = FString::JoinBy(InRemovedObjects, TEXT(", "), [](UObject* Object)
+	{
+		return ensure(Object) ? Object->GetPathName() : TEXT("null");
+	});
+	UE_LOG(LogTakesCore, Warning, TEXT("While saving %s, the following references were cleared because they were private and external: %s"),
+		*ReferencingObject->GetPathName(), *List
+		);
+}
+
+
 /** Helper function - get the first PIE world (or first PIE client world if there is more than one) */
 UWorld* GetFirstPIEWorld()
 {
@@ -95,6 +142,21 @@ void ResetViewAndWorkRange(UMovieScene* InMovieScene)
 	EditorData.ViewEnd = EditorData.WorkEnd = NewRange.GetUpperBoundValue();
 }
 
+void CleanupExternalReferences(UObject* InObject, UPackage* InPackage)
+{
+	// We want to remove any external references before calling package save because those references will
+	// prevent users from saving because of checks in the package save system. 
+	// FArchiveReplaceOrClearExternalReferences handles recursively serializing subobjects, such as components.
+	const FArchiveClearExternalReferences ReplaceActorInvalidReferences(InObject, {}, InPackage);
+	LogRemovedExternalObjects(InObject, ReplaceActorInvalidReferences.GetRemovedObjects());
+}
+
+void CleanupExternalReferences(UObject* InObject)
+{
+	UPackage* const Package = InObject->GetOutermost();
+	CleanupExternalReferences(InObject, Package);
+}
+	
 void SaveAsset(UObject* InObject)
 {
 	if (!InObject)
@@ -115,13 +177,16 @@ void SaveAsset(UObject* InObject)
 
 	double StartTime = FPlatformTime::Seconds();
 
+	CleanupExternalReferences(InObject, Package);
+
 	FSavePackageArgs SaveArgs;
 	SaveArgs.TopLevelFlags = RF_Standalone;
 	SaveArgs.SaveFlags = SAVE_Async;
 	if (bSaveTakesWithNoErrorFlag)
 	{
-		SaveArgs.SaveFlags = SaveArgs.SaveFlags | SAVE_NoError;  
+		SaveArgs.SaveFlags = SaveArgs.SaveFlags | SAVE_NoError;
 	}
+
 	if (UPackage::SavePackage(Package, NULL, *PackageFileName, SaveArgs))
 	{
 		UPackage::WaitForAsyncFileWrites();
diff --git a/Engine/Plugins/VirtualProduction/Takes/Source/TakesCore/Public/TakesUtils.h b/Engine/Plugins/VirtualProduction/Takes/Source/TakesCore/Public/TakesUtils.h
--- a/Engine/Plugins/VirtualProduction/Takes/Source/TakesCore/Public/TakesUtils.h
+++ b/Engine/Plugins/VirtualProduction/Takes/Source/TakesCore/Public/TakesUtils.h
@@ -30,6 +30,8 @@ namespace TakesUtils
 
 	TAKESCORE_API void ResetViewAndWorkRange(UMovieScene* InMovieScene);
 
+	TAKESCORE_API void CleanupExternalReferences(UObject* InObject);
+	
 	TAKESCORE_API void SaveAsset(UObject* InObject);
 
 	TAKESCORE_API void CreateCameraCutTrack(ULevelSequence* LevelSequence, const FGuid& RecordedCameraGuid, const FMovieSceneSequenceID& SequenceID, const TRange<FFrameNumber>& InRange);
@@ -173,4 +175,4 @@ namespace TakesUtils
 	/* The time at which to record. Taken from the Sequencer global time, otherwise based on timecode */
 	TAKESCORE_API FQualifiedFrameTime GetRecordTime(TSharedPtr<ISequencer> Sequencer, ULevelSequence* SequenceAsset,
 		const FTimecode& TimecodeAtStart, bool bStartAtCurrentTimecode, float TimeDilation = 1.f);
-}
\ No newline at end of file
+}
