Take Recorder: Fix genlocked engine timestep being overriden by hitch protection.

The fix is to improve user UX: disable hitch protection when any custom timestep is set.
"Enable Hitch Protection" checkbox is now greyed out when a custom timestep is set.
Users were not expecting hitch protection to kick in so we just force them to make a conscious decision when they want to use hitch protection.

There is a tooltip on the checkbox in Take Recorder settings to explain to users why the checkbox is greyed out.

diff --git a/Engine/Plugins/VirtualProduction/Takes/Source/TakeRecorder/Private/Customization/TakeRecorderProjectSettingsCustomization.h b/Engine/Plugins/VirtualProduction/Takes/Source/TakeRecorder/Private/Customization/TakeRecorderProjectSettingsCustomization.h
--- a/Engine/Plugins/VirtualProduction/Takes/Source/TakeRecorder/Private/Customization/TakeRecorderProjectSettingsCustomization.h
+++ b/Engine/Plugins/VirtualProduction/Takes/Source/TakeRecorder/Private/Customization/TakeRecorderProjectSettingsCustomization.h
@@ -11,6 +11,7 @@
 #include "Modules/ModuleManager.h"
 #include "TakeRecorderSettings.h"
 #include "TakeRecorderSourceProperty.h"
+#include "Timecode/HitchProtectionModel.h"
 
 class FTakeRecorderProjectSettingsCustomization : public IDetailCustomization
 {
@@ -83,12 +84,69 @@ private:
 		
 		// Structs have a root node, under which their properties are nested. Since we're promoting them up, we need to hide this node.
 		HitchProtectionCategory.AddProperty(HitchProtectionSettings).EditCondition(false, {}).EditConditionHides(true);
+
+		// Enabled property show should first...
+		AddHitchProtectionEnabledProperty(HitchProtectionCategory, HitchProtectionSettings);
+		
+		// ... then all the others.
+		const auto AddProperty = [&HitchProtectionCategory, &HitchProtectionSettings](FName PropertyName)
+		{
+			HitchProtectionCategory.AddProperty(HitchProtectionSettings->GetChildHandle(PropertyName))
+				.EditCondition(TAttribute<bool>::CreateLambda([]()
+				{
+					return UE::TakeRecorder::HitchProtectionModel::CanInitializeHitchProtection();
+				}),
+				// Hide the property either if !bEnableHitchProtection or when !CanInitializeHitchProtection
+				{}, ECustomEditConditionMode::Merge)
+				.EditConditionHides(true);
+		};
+		AddProperty(GET_MEMBER_NAME_CHECKED(FTakeRecorderHitchProtectionParameters, RegressionBufferSizeInSeconds));
+		AddProperty(GET_MEMBER_NAME_CHECKED(FTakeRecorderHitchProtectionParameters, CustomTimestep));
+		AddProperty(GET_MEMBER_NAME_CHECKED(FTakeRecorderHitchProtectionParameters, MaxCatchupSeconds));
+	}
+
+	void AddHitchProtectionEnabledProperty(IDetailCategoryBuilder& HitchProtectionCategory, const TSharedRef<IPropertyHandle> HitchProtectionSettings)
+	{
+		const TSharedPtr<IPropertyHandle> EnableHitchProtectionProperty = HitchProtectionSettings->GetChildHandle(
+			GET_MEMBER_NAME_CHECKED(FTakeRecorderHitchProtectionParameters, bEnableHitchProtection)
+			);
+		
+		FDetailWidgetRow& EnableRow = HitchProtectionCategory.AddProperty(EnableHitchProtectionProperty)
+			.EditCondition(TAttribute<bool>::CreateLambda([]()
+			{
+				const bool bCanEdit = UE::TakeRecorder::HitchProtectionModel::CanInitializeHitchProtection(); 
+				return bCanEdit;
+			}), {}).CustomWidget();
+
+		const TSharedRef<SWidget> NameWidget = EnableHitchProtectionProperty->CreatePropertyNameWidget();
+		ForEachChild(*NameWidget, [DefaultText = EnableHitchProtectionProperty->GetToolTipText()](SWidget& Widget)
+		{
+			Widget.SetToolTipText(TAttribute<FText>::CreateLambda([DefaultText]
+			{
+				return UE::TakeRecorder::HitchProtectionModel::CanInitializeHitchProtection()
+					? DefaultText
+					: UE::TakeRecorder::HitchProtectionModel::GetHitchProtectionDisabledReasonText();
+			}));
+		});
 		
-		uint32 NumChildren = 0;
-		HitchProtectionSettings->GetNumChildren(NumChildren);
-		for (uint32 Index = 0; Index < NumChildren; ++Index)
+		EnableRow.NameContent() [ NameWidget ];
+		EnableRow.ValueContent() [ EnableHitchProtectionProperty->CreatePropertyValueWidget() ];
+	}
+
+	template<typename TCallback>
+	static void ForEachChild(SWidget& Widget, TCallback&& Callback)
+	{
+		Callback(Widget);
+
+		FChildren* Children = Widget.GetChildren();
+		if (!Children)
 		{
-			HitchProtectionCategory.AddProperty(HitchProtectionSettings->GetChildHandle(Index));
+			return;
 		}
+
+		Children->ForEachWidget([&Callback](SWidget& Child)
+		{
+			ForEachChild(Child, Callback);
+		});
 	}
 };
diff --git a/Engine/Plugins/VirtualProduction/Takes/Source/TakeRecorder/Private/Timecode/HitchProtectionModel.cpp b/Engine/Plugins/VirtualProduction/Takes/Source/TakeRecorder/Private/Timecode/HitchProtectionModel.cpp
new file mode 100644
--- /dev/null
+++ b/Engine/Plugins/VirtualProduction/Takes/Source/TakeRecorder/Private/Timecode/HitchProtectionModel.cpp
@@ -0,0 +1,27 @@
+﻿// Copyright Epic Games, Inc. All Rights Reserved.
+
+#include "HitchProtectionModel.h"
+
+#include "Engine/Engine.h"
+
+namespace UE::TakeRecorder
+{
+bool HitchProtectionModel::CanInitializeHitchProtection()
+{
+	if (!GEngine)
+	{
+		return false;
+	}
+
+	UEngineCustomTimeStep* CustomTimeStep = GEngine->GetCustomTimeStep();
+	
+	// If a custom time step is set, we cannot run hitch protection because it requires overriding the timestep.
+	// We don't want to interrupt a synchronized genlock when the user starts recording.
+	return CustomTimeStep == nullptr;
+}
+
+FText HitchProtectionModel::GetHitchProtectionDisabledReasonText()
+{
+	return NSLOCTEXT("TakeRecorder", "HitchProtection.CannotInitializeReason", "To use hitch protection, clear the custom engine timestep first.");
+}
+}
diff --git a/Engine/Plugins/VirtualProduction/Takes/Source/TakeRecorder/Private/Timecode/HitchProtectionModel.h b/Engine/Plugins/VirtualProduction/Takes/Source/TakeRecorder/Private/Timecode/HitchProtectionModel.h
new file mode 100644
--- /dev/null
+++ b/Engine/Plugins/VirtualProduction/Takes/Source/TakeRecorder/Private/Timecode/HitchProtectionModel.h
@@ -0,0 +1,17 @@
+﻿// Copyright Epic Games, Inc. All Rights Reserved.
+
+#pragma once
+
+class FText;
+
+/**
+ * This is a stateless model for hitch protection. 
+ */
+namespace UE::TakeRecorder::HitchProtectionModel
+{
+/** @return Whether hitch protection can be initialized. False if there is an incompatible engine timestep set. */
+bool CanInitializeHitchProtection(); 
+
+/** @return Returns why hitch protection cannot be initialized */
+FText GetHitchProtectionDisabledReasonText();
+}
diff --git a/Engine/Plugins/VirtualProduction/Takes/Source/TakeRecorder/Private/Timecode/HitchlessProtectionRootLogic.cpp b/Engine/Plugins/VirtualProduction/Takes/Source/TakeRecorder/Private/Timecode/HitchlessProtectionRootLogic.cpp
--- a/Engine/Plugins/VirtualProduction/Takes/Source/TakeRecorder/Private/Timecode/HitchlessProtectionRootLogic.cpp
+++ b/Engine/Plugins/VirtualProduction/Takes/Source/TakeRecorder/Private/Timecode/HitchlessProtectionRootLogic.cpp
@@ -2,6 +2,7 @@
 
 #include "HitchlessProtectionRootLogic.h"
 
+#include "HitchProtectionModel.h"
 #include "PropertyEditorModule.h"
 #include "Recorder/TakeRecorderHitchProtectionParameters.h"
 #include "Recorder/TakeRecorderSubsystem.h"
@@ -47,7 +48,7 @@ void FHitchlessProtectionRootLogic::SetupTakeRecorderDelegates()
 void FHitchlessProtectionRootLogic::OnInitializeRecording(UTakeRecorder* TakeRecorder)
 {
 	const FTakeRecorderHitchProtectionParameters& Settings = GetMutableDefault<UTakeRecorderProjectSettings>()->HitchProtectionSettings;
-	if (!Settings.bEnableHitchProtection)
+	if (!Settings.bEnableHitchProtection || !HitchProtectionModel::CanInitializeHitchProtection())
 	{
 		return;
 	}
diff --git a/Engine/Plugins/VirtualProduction/Takes/Source/TakeRecorder/Private/Timecode/Regression/TimecodeRegressionRecordSetup.cpp b/Engine/Plugins/VirtualProduction/Takes/Source/TakeRecorder/Private/Timecode/Regression/TimecodeRegressionRecordSetup.cpp
--- a/Engine/Plugins/VirtualProduction/Takes/Source/TakeRecorder/Private/Timecode/Regression/TimecodeRegressionRecordSetup.cpp
+++ b/Engine/Plugins/VirtualProduction/Takes/Source/TakeRecorder/Private/Timecode/Regression/TimecodeRegressionRecordSetup.cpp
@@ -9,6 +9,7 @@
 #include "Recorder/TakeRecorderHitchProtectionParameters.h"
 #include "Recorder/TakeRecorderSubsystem.h"
 #include "TimecodeNotifications.h"
+#include "Timecode/HitchProtectionModel.h"
 #include "UObject/Package.h"
 
 namespace UE::TakeRecorder
@@ -39,6 +40,12 @@ FTimecodeEstimationSetupResult FTimecodeRegressionRecordSetup::SetupEngineEnviro
 {
 	FTimecodeEstimationSetupResult Result;
 	
+	// You are supposed to check this before calling SetupEngineEnvironmentForRecording.
+	if (!ensure(HitchProtectionModel::CanInitializeHitchProtection()))
+	{
+		return Result;
+	}
+	
 	UTimecodeProvider* CurrentTimecodeProvider = GEngine ? GEngine->GetTimecodeProvider() : nullptr;
 	if (!CurrentTimecodeProvider)
 	{
