Take Recorder: Introduce bStopOnRollover setting, which causes recording to stop once timecode rolls over. This prevents the user from having a broken level sequence if timecode rolls over.

diff --git a/Engine/Plugins/VirtualProduction/Takes/Source/TakeRecorder/Private/Recorder/TakeRecorder.cpp b/Engine/Plugins/VirtualProduction/Takes/Source/TakeRecorder/Private/Recorder/TakeRecorder.cpp
--- a/Engine/Plugins/VirtualProduction/Takes/Source/TakeRecorder/Private/Recorder/TakeRecorder.cpp
+++ b/Engine/Plugins/VirtualProduction/Takes/Source/TakeRecorder/Private/Recorder/TakeRecorder.cpp
@@ -805,11 +805,48 @@ FQualifiedFrameTime UTakeRecorder::GetRecordTime() const
 		Parameters.Project.bStartAtCurrentTimecode, TimeDilation);
 }
 
+namespace UE::TakeRecorder
+{
+static bool operator<(const FTimecode& InLeft, const FTimecode& InRight)
+{
+	return InLeft.Hours < InRight.Hours
+		|| (InLeft.Hours <= InRight.Hours && InLeft.Minutes < InRight.Minutes)
+		|| (InLeft.Hours <= InRight.Hours && InLeft.Minutes <= InRight.Minutes && InLeft.Seconds < InRight.Seconds)
+		|| (InLeft.Hours <= InRight.Hours && InLeft.Minutes <= InRight.Minutes && InLeft.Seconds <= InRight.Seconds
+			&& InLeft.Frames < InRight.Frames)
+		|| (InLeft.Hours <= InRight.Hours && InLeft.Minutes <= InRight.Minutes && InLeft.Seconds <= InRight.Seconds
+			&& InLeft.Frames <= InRight.Frames && InLeft.Subframe < InRight.Subframe);
+}
+
+static void ShowRolloverNotification()
+{
+	if (FSlateApplication::IsInitialized())
+	{
+		FNotificationInfo Info(LOCTEXT("RolloverMessage.Title", "Timecode rolled over."));
+		Info.SubText = LOCTEXT("RolloverMessage.Subtext", "Recording was stopped because timecode rolled over.");
+		Info.bFireAndForget = true;
+		Info.ExpireDuration = 4.f;
+		
+		FSlateNotificationManager::Get().AddNotification(Info);
+	}
+}
+}
+
 void UTakeRecorder::InternalTick(float DeltaTime)
 {
 	SCOPED_TAKERECORDER_TRACE(UTakeRecorder_InternalTick);
 	using namespace UE::MovieScene;
 
+	using namespace UE::TakeRecorder;
+	// We assume timecode providers do not change after recording is started.
+	// If no timecode provider is set, then FApp::GetTimecode() == TimecodeAtStart == FTimecode(0,0,0,0).
+	if (Parameters.User.bStopOnRollover && FApp::GetTimecode() < TimecodeAtStart)
+	{
+		ShowRolloverNotification();
+		Stop();
+		return;
+	}
+
 	FScopedSignedObjectModifyDefer FlushOnTick(true);
 
 	TSharedPtr<ISequencer> Sequencer = WeakSequencer.Pin();
diff --git a/Engine/Plugins/VirtualProduction/Takes/Source/TakeRecorder/Private/Recorder/TakeRecorderParameters.cpp b/Engine/Plugins/VirtualProduction/Takes/Source/TakeRecorder/Private/Recorder/TakeRecorderParameters.cpp
--- a/Engine/Plugins/VirtualProduction/Takes/Source/TakeRecorder/Private/Recorder/TakeRecorderParameters.cpp
+++ b/Engine/Plugins/VirtualProduction/Takes/Source/TakeRecorder/Private/Recorder/TakeRecorderParameters.cpp
@@ -12,6 +12,7 @@ FTakeRecorderUserParameters::FTakeRecorderUserParameters()
 	, EngineTimeDilation(1.f)
 	, bResetPlayhead(true)
 	, bStopAtPlaybackEnd(false)
+	, bStopOnRollover(true)
 	, bRemoveRedundantTracks(true)
 	, ReduceKeysTolerance(KINDA_SMALL_NUMBER)
 	, bSaveRecordedAssets(false)
diff --git a/Engine/Plugins/VirtualProduction/Takes/Source/TakeRecorder/Public/Recorder/TakeRecorderParameters.h b/Engine/Plugins/VirtualProduction/Takes/Source/TakeRecorder/Public/Recorder/TakeRecorderParameters.h
--- a/Engine/Plugins/VirtualProduction/Takes/Source/TakeRecorder/Public/Recorder/TakeRecorderParameters.h
+++ b/Engine/Plugins/VirtualProduction/Takes/Source/TakeRecorder/Public/Recorder/TakeRecorderParameters.h
@@ -39,6 +39,13 @@ struct FTakeRecorderUserParameters
 	UPROPERTY(config, EditAnywhere, BlueprintReadWrite, Category = "User Settings")
 	bool bStopAtPlaybackEnd;
 
+	/**
+	 * Automatically stop recording when the current time code value loops over to before the start timecode value,
+	 * e.g. when a rollover happens (23:59:59 -> 00:00:00:01), or if you started recording at 18:39:12:01 and timecode later jumps to 18:35:00:01.
+	 */
+	UPROPERTY(config, EditAnywhere, BlueprintReadWrite, Category = "User Settings")
+	bool bStopOnRollover;
+
 	/** Recommended for use with recorded spawnables. Beware that changes to actor instances in the map after recording may alter the recording when played back */
 	UPROPERTY(config, EditAnywhere, BlueprintReadWrite, Category="User Settings")
 	bool bRemoveRedundantTracks;
