Ensure our target object is valid before checking the world for world type.  This will prevent a crash when the source is deleted during recording.

#rnx

diff --git a/Engine/Plugins/VirtualProduction/Takes/Source/TakeRecorderSources/Private/TakeRecorderActorSource.cpp b/Engine/Plugins/VirtualProduction/Takes/Source/TakeRecorderSources/Private/TakeRecorderActorSource.cpp
--- a/Engine/Plugins/VirtualProduction/Takes/Source/TakeRecorderSources/Private/TakeRecorderActorSource.cpp
+++ b/Engine/Plugins/VirtualProduction/Takes/Source/TakeRecorderSources/Private/TakeRecorderActorSource.cpp
@@ -1788,7 +1788,7 @@ bool UTakeRecorderActorSource::GetRecordToPossessable() const
 	
 	// In non-editor worlds, we should defer to the record type versus what the global spawnable specifices.  This is to address cases where we change
 	// the recording type on pre-recording. See UE-306474.
-	if (Target->GetWorld()->WorldType == EWorldType::Editor && !AllowsSpawnableObjects())
+	if (Target && Target->GetWorld() && Target->GetWorld()->WorldType == EWorldType::Editor && !AllowsSpawnableObjects())
 	{
 		return true;
 	}
