Fixing issue where TakeUtils::MakeNewAsset would incrementally append a number onto the end of a name instead of incrementing the index in place at the end of the new name eg. name_0_1_2_3_4_5 instead of name_5.

This could ultimately end in a crash when the name became too long for a saved asset at the end of a recording. This crash could occur when new components were added during a recording, such as when dragging an actor around with a specific BP owned component hierarchy. With this fix, the correct number of sequences are recorded but their names correctly reflect the index, and the name does not incrementally grow in length.

diff --git a/Engine/Plugins/VirtualProduction/Takes/Source/TakesCore/Public/TakesUtils.h b/Engine/Plugins/VirtualProduction/Takes/Source/TakesCore/Public/TakesUtils.h
--- a/Engine/Plugins/VirtualProduction/Takes/Source/TakesCore/Public/TakesUtils.h
+++ b/Engine/Plugins/VirtualProduction/Takes/Source/TakesCore/Public/TakesUtils.h
@@ -125,34 +125,34 @@ namespace TakesUtils
 	template<typename AssetType>
 	static AssetType* MakeNewAsset(const FString& BaseAssetPath, const FString& BaseAssetName)
 	{
-		const FString Dot(TEXT("."));
+		auto MakeAssetObjectPath = [](const FString& BaseAssetPath, const FString& AssetName) -> FSoftObjectPath
+			{
+				return FSoftObjectPath((BaseAssetPath / AssetName) + TEXT(".") + AssetName);
+			};
+
 		FString SanitizedBaseAssetPath = UPackageTools::SanitizePackageName(BaseAssetPath);
 		FString SanitizedBaseAssetName = UPackageTools::SanitizePackageName(BaseAssetName);
 
-		FString AssetPath = SanitizedBaseAssetPath;
-		FString AssetName = SanitizedBaseAssetName;
-
-		AssetName = FPaths::MakeValidFileName(AssetName);
-
-		AssetPath /= AssetName;
-		AssetPath += Dot + AssetName;
+		FString SourceAssetName = FPaths::MakeValidFileName(SanitizedBaseAssetName);
+		FSoftObjectPath AssetObjectPath = MakeAssetObjectPath(SanitizedBaseAssetPath, SourceAssetName);
 
 		FAssetRegistryModule& AssetRegistryModule = FModuleManager::LoadModuleChecked<FAssetRegistryModule>(TEXT("AssetRegistry"));
-		FAssetData AssetData = AssetRegistryModule.Get().GetAssetByObjectPath(FSoftObjectPath(AssetPath));
+		FAssetData AssetData = AssetRegistryModule.Get().GetAssetByObjectPath(AssetObjectPath);
 
 		// if object with same name exists, try a different name until we don't find one
 		int32 ExtensionIndex = 0;
+		FString AssetName = SourceAssetName;
 		while (AssetData.IsValid() && AssetData.GetClass() == AssetType::StaticClass())
 		{
-			AssetName = FString::Printf(TEXT("%s_%d"), *AssetName, ExtensionIndex);
-			AssetPath = (BaseAssetPath / AssetName) + Dot + AssetName;
-			AssetData = AssetRegistryModule.Get().GetAssetByObjectPath(FSoftObjectPath(AssetPath));
+			AssetName = FString::Printf(TEXT("%s_%d"), *SourceAssetName, ExtensionIndex);
+			AssetObjectPath = MakeAssetObjectPath(SanitizedBaseAssetPath, AssetName);
+			AssetData = AssetRegistryModule.Get().GetAssetByObjectPath(AssetObjectPath);
 
 			ExtensionIndex++;
 		}
 
 		// Create the new asset in the package we just made
-		AssetPath = (SanitizedBaseAssetPath / AssetName);
+		const FString AssetPath = (SanitizedBaseAssetPath / AssetName);
 
 		FString FileName;
 		if (FPackageName::TryConvertLongPackageNameToFilename(AssetPath, FileName))
