When recording objects as spawnable, the source object must be serialized in whole into the level sequence package.  Any references to the external world could leak into the level sequence that we intend to save.

This change extends the saving code for take recorder to check to see if any external references and clear those references out prior to saving the package to avoid any checks in the save package call where it is checking for external references.  It skips this path when the user is recording possessables.

diff --git a/Engine/Plugins/VirtualProduction/CinematicAssemblyTools/Source/CineAssemblyToolsEditor/Private/TakeRecorder/CineAssemblyTakeRecorderIntegration.cpp b/Engine/Plugins/VirtualProduction/CinematicAssemblyTools/Source/CineAssemblyToolsEditor/Private/TakeRecorder/CineAssemblyTakeRecorderIntegration.cpp
--- a/Engine/Plugins/VirtualProduction/CinematicAssemblyTools/Source/CineAssemblyToolsEditor/Private/TakeRecorder/CineAssemblyTakeRecorderIntegration.cpp
+++ b/Engine/Plugins/VirtualProduction/CinematicAssemblyTools/Source/CineAssemblyToolsEditor/Private/TakeRecorder/CineAssemblyTakeRecorderIntegration.cpp
@@ -10,9 +10,12 @@
 #include "ProductionSettings.h"
 #include "ITakeRecorderModule.h"
 #include "Recorder/TakeRecorder.h"
+#include "TakeRecorderSources.h"
 #include "Sections/MovieSceneSubSection.h"
 #include "TakesUtils.h"
 
+DEFINE_LOG_CATEGORY_STATIC(LogCineAssemblyTakeRecorderIntegration, Log, All)
+
 FCineAssemblyTakeRecorderIntegration::FCineAssemblyTakeRecorderIntegration()
 {
 	UCineAssemblyTakeRecorderSettings* Settings = GetMutableDefault<UCineAssemblyTakeRecorderSettings>();
@@ -153,7 +156,18 @@ void FCineAssemblyTakeRecorderIntegration::OnRecordingStopped(UTakeRecorder* Tak
 						MovieScene->SetReadOnly(true);
 					}
 
-					TakesUtils::SaveAsset(SubSequence);
+					if (UTakeRecorderSources* Sources = Assembly->FindMetaData<UTakeRecorderSources>())
+					{
+						FTakeRecorderSourcesSettings Settings = Sources->GetSettings();
+						if (Settings.bSaveRecordedAssets)
+						{
+							TakesUtils::SaveAsset(SubSequence, Settings.bRecordToPossessable);
+						}
+					}
+					else
+					{
+						UE_LOG(LogCineAssemblyTakeRecorderIntegration, Error, TEXT("No sources found on stop recording. Unable determine how to save %s."), *Assembly->GetFName().ToString());
+					}
 				}
 			}
 		}
diff --git a/Engine/Plugins/VirtualProduction/Takes/Source/TakeRecorder/Private/Recorder/TakeRecorder.cpp b/Engine/Plugins/VirtualProduction/Takes/Source/TakeRecorder/Private/Recorder/TakeRecorder.cpp
--- a/Engine/Plugins/VirtualProduction/Takes/Source/TakeRecorder/Private/Recorder/TakeRecorder.cpp
+++ b/Engine/Plugins/VirtualProduction/Takes/Source/TakeRecorder/Private/Recorder/TakeRecorder.cpp
@@ -1254,6 +1254,12 @@ void UTakeRecorder::StopInternal(const bool bCancelled)
 			{
 				if (Sequencer)
 				{
+					if (!Parameters.Project.bRecordToPossessable)
+					{
+						// Ensure we don't have any external references prior to saving.
+						TakesUtils::CleanupExternalReferences(SequenceAsset);
+					}
+					
 					// Sequencer has a feature for scrubbing to a particular frame and capturing it as thumbnail.
 					// In 5.5, this was added but hidden behind a CVar for stability reason.
 					// In BaseTakes.ini, by default, we configure our Sequencer instance to capture the MiddleFrame.
@@ -1264,7 +1270,7 @@ void UTakeRecorder::StopInternal(const bool bCancelled)
 				else
 				{
 					// Fallback to "normal" saving if sequencer is not available (seems like it should be though?). Asset won't have any thumbnail.
-					TakesUtils::SaveAsset(SequenceAsset);
+					TakesUtils::SaveAsset(SequenceAsset, Parameters.Project.bRecordToPossessable);
 				}
 			}
 
diff --git a/Engine/Plugins/VirtualProduction/Takes/Source/TakeRecorderSources/Private/TakeRecorderMicrophoneAudioSource.cpp b/Engine/Plugins/VirtualProduction/Takes/Source/TakeRecorderSources/Private/TakeRecorderMicrophoneAudioSource.cpp
--- a/Engine/Plugins/VirtualProduction/Takes/Source/TakeRecorderSources/Private/TakeRecorderMicrophoneAudioSource.cpp
+++ b/Engine/Plugins/VirtualProduction/Takes/Source/TakeRecorderSources/Private/TakeRecorderMicrophoneAudioSource.cpp
@@ -422,7 +422,7 @@ TArray<UTakeRecorderSource*> UTakeRecorderMicrophoneAudioSource::PostRecording(U
 
 				if ((Sources->GetSettings().bSaveRecordedAssets) || GEditor == nullptr)
 				{
-					TakesUtils::SaveAsset(SoundWave);
+					TakesUtils::SaveAsset(SoundWave, Sources->GetSettings().bRecordToPossessable);
 				}
 			}
 		}
diff --git a/Engine/Plugins/VirtualProduction/Takes/Source/TakeTrackRecorders/Private/MovieSceneAnimationTrackRecorder.cpp b/Engine/Plugins/VirtualProduction/Takes/Source/TakeTrackRecorders/Private/MovieSceneAnimationTrackRecorder.cpp
--- a/Engine/Plugins/VirtualProduction/Takes/Source/TakeTrackRecorders/Private/MovieSceneAnimationTrackRecorder.cpp
+++ b/Engine/Plugins/VirtualProduction/Takes/Source/TakeTrackRecorders/Private/MovieSceneAnimationTrackRecorder.cpp
@@ -233,7 +233,7 @@ void UMovieSceneAnimationTrackRecorder::FinalizeTrackImpl()
 	FTrackRecorderSettings TrackRecorderSettings = OwningTakeRecorderSource->GetTrackRecorderSettings();
 	if (TrackRecorderSettings.bSaveRecordedAssets)
 	{
-		TakesUtils::SaveAsset(GetAnimSequence());
+		TakesUtils::SaveAsset(GetAnimSequence(), TrackRecorderSettings.bRecordToPossessable);
 	}
 }
 
diff --git a/Engine/Plugins/VirtualProduction/Takes/Source/TakesCore/Private/TakeRecorderSources.cpp b/Engine/Plugins/VirtualProduction/Takes/Source/TakesCore/Private/TakeRecorderSources.cpp
--- a/Engine/Plugins/VirtualProduction/Takes/Source/TakesCore/Private/TakeRecorderSources.cpp
+++ b/Engine/Plugins/VirtualProduction/Takes/Source/TakesCore/Private/TakeRecorderSources.cpp
@@ -567,7 +567,7 @@ void UTakeRecorderSources::StopRecording(class ULevelSequence* InSequence, const
 			}
 			else if (Settings.bSaveRecordedAssets)
 			{
-				TakesUtils::SaveAsset(SourceSubSequence.Value);
+				TakesUtils::SaveAsset(SourceSubSequence.Value, Settings.bRecordToPossessable);
 			}
 		}
 	}
diff --git a/Engine/Plugins/VirtualProduction/Takes/Source/TakesCore/Private/TakesUtils.cpp b/Engine/Plugins/VirtualProduction/Takes/Source/TakesCore/Private/TakesUtils.cpp
--- a/Engine/Plugins/VirtualProduction/Takes/Source/TakesCore/Private/TakesUtils.cpp
+++ b/Engine/Plugins/VirtualProduction/Takes/Source/TakesCore/Private/TakesUtils.cpp
@@ -17,6 +17,7 @@
 #include "Sections/MovieSceneCameraCutSection.h"
 #include "UObject/SavePackage.h"
 
+#include "Serialization/ArchiveReplaceOrClearExternalReferences.h"
 #define LOCTEXT_NAMESPACE "TakesUtils"
 
 namespace TakesUtils
@@ -25,6 +26,52 @@ namespace TakesUtils
 static bool bSaveTakesWithNoErrorFlag = false;
 static FAutoConsoleVariableRef  CVarSaveTakesWithNoError(TEXT("Takes.SaveTakesWithNoErrorFlag"), bSaveTakesWithNoErrorFlag, TEXT("For DrySquash only."));
 
+class FArchiveClearExternalReferences : public FArchiveReplaceOrClearExternalReferences<UObject>
+{
+	using Super = FArchiveReplaceOrClearExternalReferences<UObject>;
+
+	TArray<UObject*> RemovedObjects;
+public:
+	
+	FArchiveClearExternalReferences(UObject* InSearchObject, const TMap<UObject*, UObject*>& InReplacementMap, UPackage* InDestPackage)
+		: Super(InSearchObject, InReplacementMap, InDestPackage, EArchiveReplaceObjectFlags::NullPrivateRefs | EArchiveReplaceObjectFlags::DelayStart )
+	{
+		SerializeSearchObject();
+	}
+
+	const TArray<UObject*>& GetRemovedObjects() const { return RemovedObjects; }
+
+	virtual FArchive& operator<<(UObject*& Obj) override
+	{
+		UObject* Original = Obj;
+		Super::operator<<(Obj);
+
+		if (Original && !Obj)
+		{
+			RemovedObjects.AddUnique(Original);
+		}
+		
+		return *this;
+	}
+};
+
+static void LogRemovedExternalObjects(UObject* ReferencingObject, const TArray<UObject*>& InRemovedObjects)
+{
+	if (InRemovedObjects.IsEmpty())
+	{
+		return;
+	}
+
+	const FString List = FString::JoinBy(InRemovedObjects, TEXT(", "), [](UObject* Object)
+	{
+		return ensure(Object) ? Object->GetPathName() : TEXT("null");
+	});
+	UE_LOG(LogTakesCore, Warning, TEXT("While saving %s, the following references were cleared because they were private and external: %s"),
+		*ReferencingObject->GetPathName(), *List
+		);
+}
+
+
 /** Helper function - get the first PIE world (or first PIE client world if there is more than one) */
 UWorld* GetFirstPIEWorld()
 {
@@ -95,7 +142,22 @@ void ResetViewAndWorkRange(UMovieScene* InMovieScene)
 	EditorData.ViewEnd = EditorData.WorkEnd = NewRange.GetUpperBoundValue();
 }
 
-void SaveAsset(UObject* InObject)
+void CleanupExternalReferences(UObject* InObject, UPackage* InPackage)
+{
+	// We want to remove any external references before calling package save because those references will
+	// prevent users from saving because of checks in the package save system. 
+	// FArchiveReplaceOrClearExternalReferences handles recursively serializing subobjects, such as components.
+	const FArchiveClearExternalReferences ReplaceActorInvalidReferences(InObject, {}, InPackage);
+	LogRemovedExternalObjects(InObject, ReplaceActorInvalidReferences.GetRemovedObjects());
+}
+
+void CleanupExternalReferences(UObject* InObject)
+{
+	UPackage* const Package = InObject->GetOutermost();
+	CleanupExternalReferences(InObject, Package);
+}
+	
+void SaveAsset(UObject* InObject, bool bIsPossessable)
 {
 	if (!InObject)
 	{
@@ -115,13 +177,19 @@ void SaveAsset(UObject* InObject)
 
 	double StartTime = FPlatformTime::Seconds();
 
+	if (!bIsPossessable)
+	{
+		CleanupExternalReferences(InObject, Package);		
+	}
+
 	FSavePackageArgs SaveArgs;
 	SaveArgs.TopLevelFlags = RF_Standalone;
 	SaveArgs.SaveFlags = SAVE_Async;
 	if (bSaveTakesWithNoErrorFlag)
 	{
-		SaveArgs.SaveFlags = SaveArgs.SaveFlags | SAVE_NoError;  
+		SaveArgs.SaveFlags = SaveArgs.SaveFlags | SAVE_NoError;
 	}
+
 	if (UPackage::SavePackage(Package, NULL, *PackageFileName, SaveArgs))
 	{
 		UPackage::WaitForAsyncFileWrites();
diff --git a/Engine/Plugins/VirtualProduction/Takes/Source/TakesCore/Public/TakesUtils.h b/Engine/Plugins/VirtualProduction/Takes/Source/TakesCore/Public/TakesUtils.h
--- a/Engine/Plugins/VirtualProduction/Takes/Source/TakesCore/Public/TakesUtils.h
+++ b/Engine/Plugins/VirtualProduction/Takes/Source/TakesCore/Public/TakesUtils.h
@@ -30,7 +30,9 @@ namespace TakesUtils
 
 	TAKESCORE_API void ResetViewAndWorkRange(UMovieScene* InMovieScene);
 
-	TAKESCORE_API void SaveAsset(UObject* InObject);
+	TAKESCORE_API void CleanupExternalReferences(UObject* InObject);
+	
+	TAKESCORE_API void SaveAsset(UObject* InObject, bool bIsPossessable = true);
 
 	TAKESCORE_API void CreateCameraCutTrack(ULevelSequence* LevelSequence, const FGuid& RecordedCameraGuid, const FMovieSceneSequenceID& SequenceID, const TRange<FFrameNumber>& InRange);
 
@@ -173,4 +175,4 @@ namespace TakesUtils
 	/* The time at which to record. Taken from the Sequencer global time, otherwise based on timecode */
 	TAKESCORE_API FQualifiedFrameTime GetRecordTime(TSharedPtr<ISequencer> Sequencer, ULevelSequence* SequenceAsset,
 		const FTimecode& TimecodeAtStart, bool bStartAtCurrentTimecode, float TimeDilation = 1.f);
-}
\ No newline at end of file
+}
