TakeRecorder: Integrate Naming Tokens slate widget with Take Recorder token boxes.

Fix a possible crash with audio sources.

diff --git a/Engine/Plugins/VirtualProduction/Takes/Source/TakeRecorder/Private/Customization/NamingTokenDirectoryPathCustomization.cpp b/Engine/Plugins/VirtualProduction/Takes/Source/TakeRecorder/Private/Customization/NamingTokenDirectoryPathCustomization.cpp
new file mode 100644
--- /dev/null
+++ b/Engine/Plugins/VirtualProduction/Takes/Source/TakeRecorder/Private/Customization/NamingTokenDirectoryPathCustomization.cpp
@@ -0,0 +1,226 @@
+// Copyright Epic Games, Inc. All Rights Reserved.
+
+#include "NamingTokenDirectoryPathCustomization.h"
+
+#include "ContentBrowserDelegates.h"
+#include "ContentBrowserModule.h"
+#include "DesktopPlatformModule.h"
+#include "DetailWidgetRow.h"
+#include "EditorDirectories.h"
+#include "Framework/Application/SlateApplication.h"
+#include "HAL/FileManager.h"
+#include "IContentBrowserSingleton.h"
+#include "IDesktopPlatform.h"
+#include "Misc/MessageDialog.h"
+#include "SNamingTokensEditableTextBox.h"
+#include "TakeRecorderNamingTokenCustomizationUtilities.h"
+
+#define LOCTEXT_NAMESPACE "FNamingTokenDirectoryPathStructCustomization"
+
+TSharedRef<IPropertyTypeCustomization> FNamingTokenDirectoryPathStructCustomization::MakeInstance()
+{
+	return MakeShared<FNamingTokenDirectoryPathStructCustomization>();
+}
+
+void FNamingTokenDirectoryPathStructCustomization::CustomizeHeader( TSharedRef<IPropertyHandle> StructPropertyHandle, class FDetailWidgetRow& HeaderRow, IPropertyTypeCustomizationUtils& StructCustomizationUtils )
+{
+	const TSharedPtr<IPropertyHandle> PathProperty = StructPropertyHandle->GetChildHandle("Path");
+
+	const bool bRelativeToGameContentDir = StructPropertyHandle->HasMetaData( TEXT("RelativeToGameContentDir") );
+	const bool bUseRelativePath = StructPropertyHandle->HasMetaData( TEXT("RelativePath") );
+	const bool bContentDir = StructPropertyHandle->HasMetaData( TEXT("ContentDir") ) || StructPropertyHandle->HasMetaData(TEXT("LongPackageName"));
+	const bool bForceShowPluginContent = StructPropertyHandle->HasMetaData(TEXT("ForceShowPluginContent"));
+	const bool bForceShowEngineContent = StructPropertyHandle->HasMetaData(TEXT("ForceShowEngineContent"));
+
+	AbsoluteGameContentDir = FPaths::ConvertRelativePathToFull(FPaths::ProjectContentDir());
+
+	if (PathProperty.IsValid())
+	{
+		TSharedPtr<SWidget> PickerWidget = nullptr;
+		
+		NamingTokensEditableTextBox = UE::TakeRecorder::NamingTokens::CreateNamingTokensWidgetForProperty(PathProperty);
+		PathProperty->SetOnPropertyValueChanged(FSimpleDelegate::CreateLambda([this]
+		{
+			if (NamingTokensEditableTextBox.IsValid())
+			{
+				NamingTokensEditableTextBox->EvaluateNamingTokens();
+			}
+		}));
+		
+		if(bContentDir)
+		{
+			PickerWidget = SAssignNew(PickerButton, SButton)
+			.ButtonStyle( FAppStyle::Get(), "HoverHintOnly" )
+			.ToolTipText( LOCTEXT("FolderComboToolTipText", "Choose a content directory") )
+			.OnClicked( this, &FNamingTokenDirectoryPathStructCustomization::OnPickContent, PathProperty.ToSharedRef(), bForceShowPluginContent, bForceShowEngineContent )
+			.ContentPadding(2.0f)
+			.ForegroundColor( FSlateColor::UseForeground() )
+			.IsFocusable(false)
+			.IsEnabled(this, &FNamingTokenDirectoryPathStructCustomization::IsBrowseEnabled, StructPropertyHandle)
+			[
+				SNew(SImage)
+				.Image(FAppStyle::GetBrush("PropertyWindow.Button_Ellipsis"))
+				.ColorAndOpacity(FSlateColor::UseForeground())
+			];
+
+		}
+		else
+		{
+			PickerWidget = SAssignNew(BrowseButton, SButton)
+			.ButtonStyle( FAppStyle::Get(), "HoverHintOnly" )
+			.ToolTipText( LOCTEXT( "FolderButtonToolTipText", "Choose a directory from this computer") )
+			.OnClicked( this, &FNamingTokenDirectoryPathStructCustomization::OnPickDirectory, PathProperty.ToSharedRef(), bRelativeToGameContentDir, bUseRelativePath )
+			.ContentPadding( 2.0f )
+			.ForegroundColor( FSlateColor::UseForeground() )
+			.IsFocusable( false )
+			.IsEnabled( this, &FNamingTokenDirectoryPathStructCustomization::IsBrowseEnabled, StructPropertyHandle )
+			[
+				SNew( SImage )
+				.Image( FAppStyle::GetBrush("PropertyWindow.Button_Ellipsis") )
+				.ColorAndOpacity( FSlateColor::UseForeground() )
+			];
+		}
+		
+		HeaderRow.ValueContent()
+		.MinDesiredWidth(125.0f)
+		.MaxDesiredWidth(600.0f)
+		[
+			SNew(SHorizontalBox)
+			+SHorizontalBox::Slot()
+			.FillWidth(1.0f)
+			.VAlign(VAlign_Center)
+			[
+				NamingTokensEditableTextBox.ToSharedRef()
+			]
+			+SHorizontalBox::Slot()
+			.AutoWidth()
+			.Padding(FMargin(4.0f, 0.0f, 0.0f, 0.0f))
+			.VAlign(VAlign_Center)
+			[
+				PickerWidget.ToSharedRef()
+			]
+		]
+		.NameContent()
+		[
+			StructPropertyHandle->CreatePropertyNameWidget()
+		];
+	}
+}
+
+void FNamingTokenDirectoryPathStructCustomization::CustomizeChildren( TSharedRef<IPropertyHandle> StructPropertyHandle, class IDetailChildrenBuilder& StructBuilder, IPropertyTypeCustomizationUtils& StructCustomizationUtils )
+{
+}
+
+FReply FNamingTokenDirectoryPathStructCustomization::OnPickContent(TSharedRef<IPropertyHandle> PropertyHandle, const bool bForceShowPluginContent, const bool bForceShowEngineContent) 
+{
+	const FContentBrowserModule& ContentBrowserModule = FModuleManager::LoadModuleChecked<FContentBrowserModule>("ContentBrowser");
+	FPathPickerConfig PathPickerConfig;
+	PropertyHandle->GetValue(PathPickerConfig.DefaultPath);
+	PathPickerConfig.bAllowContextMenu = false;
+	PathPickerConfig.OnPathSelected = FOnPathSelected::CreateSP(this, &FNamingTokenDirectoryPathStructCustomization::OnPathPicked, PropertyHandle);
+	PathPickerConfig.bForceShowEngineContent = bForceShowEngineContent;
+	PathPickerConfig.bForceShowPluginContent = bForceShowPluginContent;
+	
+	FMenuBuilder MenuBuilder(true, nullptr);
+	MenuBuilder.AddWidget(SNew(SBox)
+		.WidthOverride(300.0f)
+		.HeightOverride(300.0f)
+		[
+			ContentBrowserModule.Get().CreatePathPicker(PathPickerConfig)
+		], FText());
+	
+	PickerMenu = FSlateApplication::Get().PushMenu(PickerButton.ToSharedRef(),
+		FWidgetPath(),
+		MenuBuilder.MakeWidget(),
+		FSlateApplication::Get().GetCursorPos(),
+		FPopupTransitionEffect(FPopupTransitionEffect::ContextMenu)
+		);
+	
+	return FReply::Handled();
+}
+
+FReply FNamingTokenDirectoryPathStructCustomization::OnPickDirectory(TSharedRef<IPropertyHandle> PropertyHandle, const bool bRelativeToGameContentDir, const bool bUseRelativePath) const
+{
+	FString Directory;
+	IDesktopPlatform* DesktopPlatform = FDesktopPlatformModule::Get();
+	if (DesktopPlatform)
+	{
+
+		TSharedPtr<SWindow> ParentWindow = FSlateApplication::Get().FindWidgetWindow(BrowseButton.ToSharedRef());
+		void* ParentWindowHandle = (ParentWindow.IsValid() && ParentWindow->GetNativeWindow().IsValid()) ? ParentWindow->GetNativeWindow()->GetOSWindowHandle() : nullptr;
+
+		FString StartDirectory = FEditorDirectories::Get().GetLastDirectory(ELastDirectory::GENERIC_IMPORT);
+		if (bRelativeToGameContentDir && !IsValidPath(StartDirectory, bRelativeToGameContentDir))
+		{
+			StartDirectory = AbsoluteGameContentDir;
+		}
+
+		// Loop until; a) the user cancels (OpenDirectoryDialog returns false), or, b) the chosen path is valid (IsValidPath returns true)
+		for (;;)
+		{
+			if (DesktopPlatform->OpenDirectoryDialog(ParentWindowHandle, LOCTEXT("FolderDialogTitle", "Choose a directory").ToString(), StartDirectory, Directory))
+			{
+				FText FailureReason;
+				if (IsValidPath(Directory, bRelativeToGameContentDir, &FailureReason))
+				{
+					FEditorDirectories::Get().SetLastDirectory(ELastDirectory::GENERIC_IMPORT, Directory);
+
+					if (bRelativeToGameContentDir)
+					{
+						Directory.RightChopInline(AbsoluteGameContentDir.Len(), EAllowShrinking::No);
+					}
+					else if (bUseRelativePath)
+					{
+						Directory = IFileManager::Get().ConvertToRelativePath(*Directory);
+					}
+
+					PropertyHandle->SetValue(Directory);
+				}
+				else
+				{
+					StartDirectory = Directory;
+					FMessageDialog::Open(EAppMsgType::Ok, FailureReason);
+					continue;
+				}
+			}
+			break;
+		}
+	}
+
+	return FReply::Handled();
+}
+
+bool FNamingTokenDirectoryPathStructCustomization::IsValidPath(const FString& AbsolutePath, const bool bRelativeToGameContentDir, FText* const OutReason) const
+{
+	if(bRelativeToGameContentDir)
+	{
+		if(!AbsolutePath.StartsWith(AbsoluteGameContentDir))
+		{
+			if(OutReason)
+			{
+				*OutReason = FText::Format(LOCTEXT("Error_InvalidRootPath", "The chosen directory must be within {0}"), FText::FromString(AbsoluteGameContentDir));
+			}
+			return false;
+		}
+	}
+
+	return true;
+}
+
+void FNamingTokenDirectoryPathStructCustomization::OnPathPicked(const FString& Path, TSharedRef<IPropertyHandle> PropertyHandle)
+{
+	if (PickerMenu.IsValid())
+	{
+		PickerMenu->Dismiss();
+		PickerMenu.Reset();
+	}
+
+	PropertyHandle->SetValue(Path);
+}
+
+bool FNamingTokenDirectoryPathStructCustomization::IsBrowseEnabled(TSharedRef<IPropertyHandle> PropertyHandle) const
+{
+	return PropertyHandle->IsEditable();
+}
+
+#undef LOCTEXT_NAMESPACE
diff --git a/Engine/Plugins/VirtualProduction/Takes/Source/TakeRecorder/Private/Customization/NamingTokenDirectoryPathCustomization.h b/Engine/Plugins/VirtualProduction/Takes/Source/TakeRecorder/Private/Customization/NamingTokenDirectoryPathCustomization.h
new file mode 100644
--- /dev/null
+++ b/Engine/Plugins/VirtualProduction/Takes/Source/TakeRecorder/Private/Customization/NamingTokenDirectoryPathCustomization.h
@@ -0,0 +1,67 @@
+// Copyright Epic Games, Inc. All Rights Reserved.
+
+#pragma once
+
+#include "IPropertyTypeCustomization.h"
+
+#include "Containers/UnrealString.h"
+#include "Input/Reply.h"
+#include "Internationalization/Text.h"
+#include "Templates/SharedPointer.h"
+
+class IPropertyHandle;
+class IDetailPropertyRow;
+
+class IMenu;
+class SNamingTokensEditableTextBox;
+class SButton;
+
+/**
+ * A Naming Tokens version of the FDirectoryPathStructCustomization, mostly identical apart from using the
+ * NamingTokens box instead of a normal text box.
+ *
+ * We may want to move this under the NamingTokens plugin for use with other plugins.
+ */
+class FNamingTokenDirectoryPathStructCustomization : public IPropertyTypeCustomization
+{
+public:
+	static TSharedRef<IPropertyTypeCustomization> MakeInstance();
+
+	/** IPropertyTypeCustomization interface */
+	virtual void CustomizeHeader( TSharedRef<class IPropertyHandle> StructPropertyHandle, class FDetailWidgetRow& HeaderRow, IPropertyTypeCustomizationUtils& StructCustomizationUtils ) override;
+
+	virtual void CustomizeChildren( TSharedRef<class IPropertyHandle> StructPropertyHandle, class IDetailChildrenBuilder& StructBuilder, IPropertyTypeCustomizationUtils& StructCustomizationUtils ) override;              
+
+private:
+
+	/** Delegate used to display a directory picker */
+	FReply OnPickContent(TSharedRef<IPropertyHandle> PropertyHandle, const bool bForceShowPluginContent, const bool bForceShowEngineContent);
+
+	/** Delegate used to display a directory picker */
+	FReply OnPickDirectory(TSharedRef<IPropertyHandle> PropertyHandle, const bool bRelativeToGameContentDir, const bool bUseRelativePaths) const;
+
+	/** Check whether that the chosen path is valid */
+	bool IsValidPath(const FString& AbsolutePath, const bool bRelativeToGameContentDir, FText* const OutReason = nullptr) const;
+
+	/** Called when a path is picked from the path picker */
+	void OnPathPicked(const FString& Path, TSharedRef<IPropertyHandle> PropertyHandle);
+
+	/** Delegate to determine whether the browse button should be enabled */
+	bool IsBrowseEnabled(TSharedRef<IPropertyHandle> PropertyHandle) const;
+
+private:
+	/** The browse button widget */
+	TSharedPtr<SButton> BrowseButton;
+
+	/** The pick button widget */
+	TSharedPtr<SButton> PickerButton;
+
+	/** The pick button popup menu */
+	TSharedPtr<IMenu> PickerMenu;
+
+	/** The naming tokens text box */
+	TSharedPtr<SNamingTokensEditableTextBox> NamingTokensEditableTextBox;
+	
+	/** Absolute path to the game content directory */
+	FString AbsoluteGameContentDir;
+};
diff --git a/Engine/Plugins/VirtualProduction/Takes/Source/TakeRecorder/Private/Customization/TakeRecorderNamingTokenCustomizationUtilities.h b/Engine/Plugins/VirtualProduction/Takes/Source/TakeRecorder/Private/Customization/TakeRecorderNamingTokenCustomizationUtilities.h
new file mode 100644
--- /dev/null
+++ b/Engine/Plugins/VirtualProduction/Takes/Source/TakeRecorder/Private/Customization/TakeRecorderNamingTokenCustomizationUtilities.h
@@ -0,0 +1,102 @@
+// Copyright Epic Games, Inc. All Rights Reserved.
+
+#pragma once
+
+#include "DetailWidgetRow.h"
+#include "IDetailPropertyRow.h"
+#include "ITakeRecorderNamingTokensModule.h"
+#include "NamingTokensSpecifiers.h"
+#include "SNamingTokensEditableTextBox.h"
+#include "UObject/TextProperty.h"
+
+/**
+ * Helps customize Naming Tokens in the details panel.
+ * We may want to move this to be public under the NamingTokens plugin for use with other plugins.
+ */
+namespace UE::TakeRecorder::NamingTokens
+{
+	/** Scopes our directory path customization to naming token directory paths only. */
+	class FDirectoryPathPropertyTypeIdentifier : public IPropertyTypeIdentifier
+	{
+		virtual bool IsPropertyTypeCustomized(const IPropertyHandle& PropertyHandle) const override
+		{
+			return PropertyHandle.HasMetaData(*UE::NamingTokens::Specifiers::UseNamingTokens);
+		}
+	};
+
+	/** Creates a naming tokens widget which will respond to a property handle. */
+	inline TSharedRef<SNamingTokensEditableTextBox> CreateNamingTokensWidgetForProperty(const TSharedPtr<IPropertyHandle>& InPropertyHandle)
+	{
+		const FString TakeRecorderNamespace = ITakeRecorderNamingTokensModule::GetTakeRecorderNamespace();
+		FNamingTokenFilterArgs FilterArgs;
+		FilterArgs.AdditionalNamespacesToInclude.Add(TakeRecorderNamespace);
+		
+		return SNew(SNamingTokensEditableTextBox)
+			.ShouldEvaluateTokens(false)
+			.FilterArgs(FilterArgs)
+			.NamespaceSuggestionPriority({TakeRecorderNamespace})
+			.Text_Lambda([InPropertyHandle]() -> FText
+			{
+				FText Value;
+				if (InPropertyHandle.IsValid())
+				{
+					if (const FProperty* Property = InPropertyHandle->GetProperty())
+					{
+						if (Property->IsA<FStrProperty>())
+						{
+							FString Result;
+							InPropertyHandle->GetValue(Result);
+							Value = FText::FromString(Result);
+						}
+						else if (Property->IsA<FNameProperty>())
+						{
+							FName Result;
+							InPropertyHandle->GetValue(Result);
+							Value = FText::FromName(Result);
+						}
+						else if (Property->IsA<FTextProperty>())
+						{
+							InPropertyHandle->GetValue(Value);
+						}
+					}
+				}
+				return Value;
+			})
+			.OnTextCommitted_Lambda([InPropertyHandle](const FText& InText, ETextCommit::Type InCommitType)
+			{
+				if (InPropertyHandle.IsValid())
+				{
+					if (const FProperty* Property = InPropertyHandle->GetProperty())
+					{
+						if (Property->IsA<FStrProperty>())
+						{
+							InPropertyHandle->SetValue(InText.ToString());
+						}
+						else if (Property->IsA<FNameProperty>())
+						{
+							InPropertyHandle->SetValue(*InText.ToString());
+						}
+						else if (Property->IsA<FTextProperty>())
+						{
+							InPropertyHandle->SetValue(InText);
+						}
+					}
+				}
+			});
+	}
+
+	/** Given a detail property row create a custom naming tokens widget. */
+	inline void HandleNamingTokensRow(IDetailPropertyRow& InPropertyRow)
+	{
+		const TSharedPtr<IPropertyHandle> PropertyHandle = InPropertyRow.GetPropertyHandle();
+		InPropertyRow.CustomWidget()
+	   .NameContent()
+		[
+			PropertyHandle->CreatePropertyNameWidget()
+		]
+		.ValueContent()
+		[
+			CreateNamingTokensWidgetForProperty(PropertyHandle)
+		];
+	}
+}
diff --git a/Engine/Plugins/VirtualProduction/Takes/Source/TakeRecorder/Private/Customization/TakeRecorderProjectSettingsCustomization.h b/Engine/Plugins/VirtualProduction/Takes/Source/TakeRecorder/Private/Customization/TakeRecorderProjectSettingsCustomization.h
--- a/Engine/Plugins/VirtualProduction/Takes/Source/TakeRecorder/Private/Customization/TakeRecorderProjectSettingsCustomization.h
+++ b/Engine/Plugins/VirtualProduction/Takes/Source/TakeRecorder/Private/Customization/TakeRecorderProjectSettingsCustomization.h
@@ -3,14 +3,14 @@
 #pragma once
 
 #include "Customization/TakeRecorderAudioSettingsCustomization.h"
-#include "IDetailChildrenBuilder.h"
-#include "IDetailCustomization.h"
-#include "DetailLayoutBuilder.h"
 #include "DetailCategoryBuilder.h"
+#include "DetailLayoutBuilder.h"
 #include "DetailWidgetRow.h"
+#include "IDetailCustomization.h"
 #include "Modules/ModuleManager.h"
+#include "NamingTokensSpecifiers.h"
+#include "TakeRecorderNamingTokenCustomizationUtilities.h"
 #include "TakeRecorderSettings.h"
-#include "TakeRecorderSourceProperty.h"
 #include "Timecode/HitchProtectionModel.h"
 
 class FTakeRecorderProjectSettingsCustomization : public IDetailCustomization
@@ -21,7 +21,7 @@ class FTakeRecorderProjectSettingsCustomization : public IDetailCustomization
 
 		// Register audio device details customization for the project settings window
 		PropertyEditorModule.RegisterCustomPropertyTypeLayout(FName(TEXT("AudioInputDeviceProperty")), FOnGetPropertyTypeCustomizationInstance::CreateLambda(&MakeShared<FAudioInputDevicePropertyCustomization>));
-
+		
 		// Pop the take recorder category to the top
 		DetailLayout.EditCategory("Take Recorder");
 		// Show the Hitch Protection category after "Take Recorder"
@@ -29,7 +29,9 @@ class FTakeRecorderProjectSettingsCustomization : public IDetailCustomization
 
 		TArray<TWeakObjectPtr<UObject>> CustomizedObjects;
 		DetailLayout.GetObjectsBeingCustomized(CustomizedObjects);
-
+		
+		TArray<IDetailPropertyRow*> DetailRows;
+		
 		for (TWeakObjectPtr<UObject> EditObject : CustomizedObjects)
 		{
 			UTakeRecorderProjectSettings* Settings = Cast<UTakeRecorderProjectSettings>(EditObject.Get());
@@ -38,6 +40,37 @@ class FTakeRecorderProjectSettingsCustomization : public IDetailCustomization
 				continue;
 			}
 
+			// Get main properties
+			{
+				const TSharedRef<IPropertyHandle> SettingsProperty = DetailLayout.GetProperty(GET_MEMBER_NAME_CHECKED(UTakeRecorderProjectSettings,
+					Settings));
+
+				IDetailCategoryBuilder& SettingCategoryBuilder = DetailLayout.EditCategory(*SettingsProperty->GetDefaultCategoryName().ToString(),
+					FText::GetEmpty(), ECategoryPriority::Important);
+
+				TArray<TSharedRef<IPropertyHandle>> StructProperties;
+				uint32 NumChildren = 0;
+				SettingsProperty->GetNumChildren(NumChildren);
+
+				for (uint32 ChildIndex = 0; ChildIndex < NumChildren; ++ChildIndex)
+				{
+					TSharedPtr<IPropertyHandle> ChildHandle = SettingsProperty->GetChildHandle(ChildIndex);
+					if (ChildHandle.IsValid() && ChildHandle->IsValidHandle())
+					{
+						// Always add the property so we maintain proper order. This property has ShowOnlyInnerProperties set so this
+						// respects that behavior.
+						IDetailPropertyRow& DetailRow = SettingCategoryBuilder.AddProperty(ChildHandle);
+
+						const FStructProperty* StructProp = CastField<FStructProperty>(ChildHandle->GetProperty());
+						if (ChildHandle->HasMetaData(*UE::NamingTokens::Specifiers::UseNamingTokens)
+						&& (!StructProp || StructProp->Struct != FDirectoryPath::StaticStruct())) // Paths customized separately
+						{
+							DetailRows.Add(&DetailRow);
+						}
+					}
+				}
+			}
+
 			for (TWeakObjectPtr<UObject> WeakAdditionalSettings : Settings->AdditionalSettings)
 			{
 				UObject* AdditionalSettings = WeakAdditionalSettings.Get();
@@ -47,7 +80,6 @@ class FTakeRecorderProjectSettingsCustomization : public IDetailCustomization
 				}
 
 				UClass* Class = AdditionalSettings->GetClass();
-
 				TArray<FProperty*> EditProperties;
 				for (FProperty* Property : TFieldRange<FProperty>(Class))
 				{
@@ -62,13 +94,24 @@ class FTakeRecorderProjectSettingsCustomization : public IDetailCustomization
 					IDetailCategoryBuilder& Category = DetailLayout.EditCategory(*Class->GetDisplayNameText().ToString());
 
 					TArray<UObject*> SettingAsArray = { AdditionalSettings };
-					for (FProperty* Property : EditProperties)
+					for (const FProperty* Property : EditProperties)
 					{
-						Category.AddExternalObjectProperty(SettingAsArray, Property->GetFName());
+						if (IDetailPropertyRow* DetailRow = Category.AddExternalObjectProperty(SettingAsArray, Property->GetFName()))
+						{
+							if (Property->HasMetaData(*UE::NamingTokens::Specifiers::UseNamingTokens))
+							{
+								DetailRows.Add(DetailRow);
+							}
+						}
 					}
 				}
 			}
 		}
+
+		for (IDetailPropertyRow* DetailRow : DetailRows)
+		{
+			UE::TakeRecorder::NamingTokens::HandleNamingTokensRow(*DetailRow);
+		}
 	}
 
 private:
diff --git a/Engine/Plugins/VirtualProduction/Takes/Source/TakeRecorder/Private/TakeRecorderModule.cpp b/Engine/Plugins/VirtualProduction/Takes/Source/TakeRecorder/Private/TakeRecorderModule.cpp
--- a/Engine/Plugins/VirtualProduction/Takes/Source/TakeRecorder/Private/TakeRecorderModule.cpp
+++ b/Engine/Plugins/VirtualProduction/Takes/Source/TakeRecorder/Private/TakeRecorderModule.cpp
@@ -5,6 +5,7 @@
 #include "TakesCoreLog.h"
 #include "TakeRecorderSettings.h"
 #include "SequencerUtilities.h"
+#include "Customization/NamingTokenDirectoryPathCustomization.h"
 #include "Customization/TakeRecorderProjectSettingsCustomization.h"
 #include "Modules/ModuleManager.h"
 #include "LevelEditor.h"
@@ -362,6 +363,7 @@ void FTakeRecorderModule::StartupModule()
 	{
 		GEditor->OnEditorClose().AddRaw(this, &FTakeRecorderModule::OnEditorClose);
 	}
+	
 #endif
 }
 
@@ -453,6 +455,10 @@ void FTakeRecorderModule::RegisterDetailCustomizations()
 			UTakeRecorderProjectSettings::StaticClass()->GetFName(),
 			FOnGetDetailCustomizationInstance::CreateLambda(&MakeShared<FTakeRecorderProjectSettingsCustomization>)
 			);
+
+		PropertyEditorModule.RegisterCustomPropertyTypeLayout(FName(TEXT("DirectoryPath")),
+			FOnGetPropertyTypeCustomizationInstance::CreateLambda(&FNamingTokenDirectoryPathStructCustomization::MakeInstance),
+			MakeShared<UE::TakeRecorder::NamingTokens::FDirectoryPathPropertyTypeIdentifier>());
 	}
 #endif
 }
@@ -464,6 +470,8 @@ void FTakeRecorderModule::UnregisterDetailCustomizations()
 		if (FPropertyEditorModule* PropertyEditorModule = FModuleManager::Get().GetModulePtr<FPropertyEditorModule>("PropertyEditor"))
 		{
 			PropertyEditorModule->UnregisterCustomClassLayout(UTakeRecorderProjectSettings::StaticClass()->GetFName());
+			PropertyEditorModule->UnregisterCustomPropertyTypeLayout(FName(TEXT("DirectoryPath")),
+				MakeShared<UE::TakeRecorder::NamingTokens::FDirectoryPathPropertyTypeIdentifier>());
 		}
 	}
 }
diff --git a/Engine/Plugins/VirtualProduction/Takes/Source/TakeRecorder/TakeRecorder.Build.cs b/Engine/Plugins/VirtualProduction/Takes/Source/TakeRecorder/TakeRecorder.Build.cs
--- a/Engine/Plugins/VirtualProduction/Takes/Source/TakeRecorder/TakeRecorder.Build.cs
+++ b/Engine/Plugins/VirtualProduction/Takes/Source/TakeRecorder/TakeRecorder.Build.cs
@@ -39,6 +39,7 @@ public class TakeRecorder : ModuleRules
 				"ToolWidgets",
 				"NamingTokens",
 				"SequencerCore",
+				"NamingTokensUI",
 			}
 		);
 
diff --git a/Engine/Plugins/VirtualProduction/Takes/Source/TakeRecorderSources/Private/TakeRecorderMicrophoneAudioSource.cpp b/Engine/Plugins/VirtualProduction/Takes/Source/TakeRecorderSources/Private/TakeRecorderMicrophoneAudioSource.cpp
--- a/Engine/Plugins/VirtualProduction/Takes/Source/TakeRecorderSources/Private/TakeRecorderMicrophoneAudioSource.cpp
+++ b/Engine/Plugins/VirtualProduction/Takes/Source/TakeRecorderSources/Private/TakeRecorderMicrophoneAudioSource.cpp
@@ -135,16 +135,18 @@ TUniquePtr<TArray<bool>> UTakeRecorderMicrophoneAudioSource::GetChannelsInUse(co
 	TUniquePtr<TArray<bool>> ChannelsInUse = MakeUnique<TArray<bool>>();
 	ChannelsInUse->Init(false, InDeviceChannelCount);
 
-	UTakeRecorderSources* SourcesList = GetTypedOuter<UTakeRecorderSources>();
-	for (UTakeRecorderSource* Source : SourcesList->GetSources())
+	if (UTakeRecorderSources* SourcesList = GetTypedOuter<UTakeRecorderSources>())
 	{
-		if (UTakeRecorderMicrophoneAudioSource* MicSource = Cast<UTakeRecorderMicrophoneAudioSource>(Source))
+		for (UTakeRecorderSource* Source : SourcesList->GetSources())
 		{
-			int32 MicSourceChannelNumber = MicSource->AudioChannel.AudioInputDeviceChannel;
-			int32 ChannelIndex = MicSourceChannelNumber - 1;
-			if (ChannelIndex >= 0 && ChannelIndex < InDeviceChannelCount)
+			if (UTakeRecorderMicrophoneAudioSource* MicSource = Cast<UTakeRecorderMicrophoneAudioSource>(Source))
 			{
-				(*ChannelsInUse)[ChannelIndex] = true;
+				int32 MicSourceChannelNumber = MicSource->AudioChannel.AudioInputDeviceChannel;
+				int32 ChannelIndex = MicSourceChannelNumber - 1;
+				if (ChannelIndex >= 0 && ChannelIndex < InDeviceChannelCount)
+				{
+					(*ChannelsInUse)[ChannelIndex] = true;
+				}
 			}
 		}
 	}
