TakeRecorder: Fix assert when adjusting the Take while the Naming Tokens dropdown suggestion box is open.

diff --git a/Engine/Plugins/VirtualProduction/Takes/Source/TakeRecorder/Private/Customization/TakeRecorderNamingTokenCustomizationUtilities.h b/Engine/Plugins/VirtualProduction/Takes/Source/TakeRecorder/Private/Customization/TakeRecorderNamingTokenCustomizationUtilities.h
--- a/Engine/Plugins/VirtualProduction/Takes/Source/TakeRecorder/Private/Customization/TakeRecorderNamingTokenCustomizationUtilities.h
+++ b/Engine/Plugins/VirtualProduction/Takes/Source/TakeRecorder/Private/Customization/TakeRecorderNamingTokenCustomizationUtilities.h
@@ -62,7 +62,7 @@ namespace UE::TakeRecorder::NamingTokens
 				}
 				return Value;
 			})
-			.OnTextCommitted_Lambda([InPropertyHandle](const FText& InText, ETextCommit::Type InCommitType)
+			.OnTextChanged_Lambda([InPropertyHandle](const FText& InText)
 			{
 				if (InPropertyHandle.IsValid())
 				{
