TakeRecorder: Handle invalid metadata in Naming Tokens.

diff --git a/Engine/Plugins/VirtualProduction/Takes/Source/TakeRecorderNamingTokens/Private/TakeRecorderNamingTokens.cpp b/Engine/Plugins/VirtualProduction/Takes/Source/TakeRecorderNamingTokens/Private/TakeRecorderNamingTokens.cpp
--- a/Engine/Plugins/VirtualProduction/Takes/Source/TakeRecorderNamingTokens/Private/TakeRecorderNamingTokens.cpp
+++ b/Engine/Plugins/VirtualProduction/Takes/Source/TakeRecorderNamingTokens/Private/TakeRecorderNamingTokens.cpp
@@ -91,7 +91,7 @@ void UTakeRecorderNamingTokens::OnCreateDefaultTokens(TArray<FNamingTokenData>&
 	    FNamingTokenData::FTokenProcessorDelegateNative::CreateLambda([this]{
 	        const UMovieSceneToolsProjectSettings* ProjectSettings = GetDefault<UMovieSceneToolsProjectSettings>();
 	        const int32 TakeNumDigits = ProjectSettings->TakeNumDigits;
-	        return FText::FromString(FString::Printf(TEXT("%0*d"), TakeNumDigits, TakeMetaData->GetTakeNumber()));
+	        return FText::FromString(FString::Printf(TEXT("%0*d"), TakeNumDigits, TakeMetaData.IsValid() ? TakeMetaData->GetTakeNumber() : 0));
 	    })
 	});
 
@@ -99,7 +99,7 @@ void UTakeRecorderNamingTokens::OnCreateDefaultTokens(TArray<FNamingTokenData>&
 	    TEXT("slate"),
 	    LOCTEXT("TokenSlate", "Slate"),
 	    FNamingTokenData::FTokenProcessorDelegateNative::CreateLambda([this]{
-	        return FText::FromString(TakeMetaData->GetSlate());
+	        return FText::FromString(TakeMetaData.IsValid() ? TakeMetaData->GetSlate() : TEXT("Invalid"));
 	    })
 	});
 
