TakeRecorder: Fix duplicate slate/take tokens appearing in the UI depending on casing.
Fix multiple entries showing up in the dropdown from typing in new tokens. Only active tokens will display in the dropdown now.

#Synced-CL 48223680

diff --git a/Engine/Plugins/VirtualProduction/Takes/Source/TakeRecorder/Private/Widgets/STakeRecorderCockpit.cpp b/Engine/Plugins/VirtualProduction/Takes/Source/TakeRecorder/Private/Widgets/STakeRecorderCockpit.cpp
--- a/Engine/Plugins/VirtualProduction/Takes/Source/TakeRecorder/Private/Widgets/STakeRecorderCockpit.cpp
+++ b/Engine/Plugins/VirtualProduction/Takes/Source/TakeRecorder/Private/Widgets/STakeRecorderCockpit.cpp
@@ -157,7 +157,13 @@ namespace UE::TakeRecorder::Private
 	/** Verify the token key is allowed to be user defined. */
 	bool CanTokenBeUserDefined(const FString& InTokenKey)
 	{
-		return !InTokenKey.Equals(TokenKeySlate) && !InTokenKey.Equals(TokenKeyTake);
+		return !InTokenKey.Equals(TokenKeySlate, ESearchCase::IgnoreCase) && !InTokenKey.Equals(TokenKeyTake, ESearchCase::IgnoreCase);
+	}
+	
+	/** The token provided matches our own tokens exactly. */
+	bool IsExplicitlyOurToken(const FString& InTokenKey)
+	{
+		return InTokenKey.Equals(TokenKeyTake, ESearchCase::CaseSensitive) || InTokenKey.Equals(TokenKeySlate, ESearchCase::CaseSensitive);
 	}
 }
 
@@ -1478,21 +1484,37 @@ void STakeRecorderCockpit::CreateUserTokensUI()
 		}
 
 		// Look for new unidentified tokens that a user can define.
+		TMap<FNamingTokenData, FText>& ExistingUserDefinedTokens = GetNamingTokensData()->UserDefinedTokens;
+		TMap<FNamingTokenData, FText> ActiveUserTokens;
 		for (const FTakeRecorderNamingTokensFieldMapping& UndefinedTokens : GetNamingTokensData()->FieldToUndefinedKeys)
 		{
 			for (const FString& UndefinedTokenStr : UndefinedTokens.UndefinedKeys)
 			{
+				bool bUserTokenAdded = false;
 				FNamingTokenData UndefinedToken(UndefinedTokenStr);
-				TMap<FNamingTokenData, FText>& UserDefinedTokens = GetNamingTokensData()->UserDefinedTokens;
-				if (!UserDefinedTokens.Contains(UndefinedToken)
+				if (!ActiveUserTokens.Contains(UndefinedToken)
 					&& UE::TakeRecorder::Private::CanTokenBeUserDefined(UndefinedTokenStr))
 				{
 					// Check for existence so we don't overwrite the user defined value.
-					UserDefinedTokens.Add(UndefinedToken);
+					ActiveUserTokens.Add(UndefinedToken);
+					bUserTokenAdded = true;
+					
+					// Restore any previous values.
+					if (FText* ExistingTokenValue = ExistingUserDefinedTokens.Find(UndefinedToken))
+					{
+						ActiveUserTokens[UndefinedToken] = *ExistingTokenValue;
+					}
+				}
+				if (bUserTokenAdded || UE::TakeRecorder::Private::IsExplicitlyOurToken(UndefinedToken.TokenKey))
+				{
+					GetNamingTokensData()->VisibleUserTokens.Add(UndefinedToken.TokenKey);
 				}
-				GetNamingTokensData()->VisibleUserTokens.Add(UndefinedToken.TokenKey);
 			}
 		}
+		
+		// Overwrite all user tokens with active tokens. This destroys any custom tokens not in use, preventing the token
+		// suggestion dropdown from being cluttered with stale entries, since each letter typed technically adds a new token.
+		ExistingUserDefinedTokens = MoveTemp(ActiveUserTokens);
 	}
 
 	// Make sure our take dir token example is up-to-date.
