Cherrypick from Dev-Cine-Sparks for crash encountered during a take record.

#rnx

#Synced-CL 48868234

diff --git a/Engine/Plugins/VirtualProduction/Takes/Source/TakesCore/Private/TakesUtils.cpp b/Engine/Plugins/VirtualProduction/Takes/Source/TakesCore/Private/TakesUtils.cpp
--- a/Engine/Plugins/VirtualProduction/Takes/Source/TakesCore/Private/TakesUtils.cpp
+++ b/Engine/Plugins/VirtualProduction/Takes/Source/TakesCore/Private/TakesUtils.cpp
@@ -190,6 +190,11 @@ void SaveAsset(UObject* InObject, bool bIsPossessable)
 		SaveArgs.SaveFlags = SaveArgs.SaveFlags | SAVE_NoError;
 	}
 
+	if (!Package->IsFullyLoaded())
+	{
+		Package->FullyLoad();
+	}
+
 	if (UPackage::SavePackage(Package, NULL, *PackageFileName, SaveArgs))
 	{
 		UPackage::WaitForAsyncFileWrites();
