TakeRecorder: Fix multiple tracks with the same name not being recorded properly.

#Synced-CL 49137568

diff --git a/Engine/Plugins/VirtualProduction/Takes/Source/TakesCore/Private/TakeRecorderSources.cpp b/Engine/Plugins/VirtualProduction/Takes/Source/TakesCore/Private/TakeRecorderSources.cpp
--- a/Engine/Plugins/VirtualProduction/Takes/Source/TakesCore/Private/TakeRecorderSources.cpp
+++ b/Engine/Plugins/VirtualProduction/Takes/Source/TakesCore/Private/TakeRecorderSources.cpp
@@ -133,7 +133,8 @@ void UTakeRecorderSources::PostEditChangeProperty(FPropertyChangedEvent& Propert
 	}
 }
 
-void UTakeRecorderSources::PreRecordingRecursive(TArray<UTakeRecorderSource*> InSources, ULevelSequence* InRootSequence, TArray<UTakeRecorderSource*>& NewSourcesOut, FManifestSerializer* InManifestSerializer)
+void UTakeRecorderSources::PreRecordingRecursive(TArray<UTakeRecorderSource*> InSources, ULevelSequence* InRootSequence,
+	TArray<UTakeRecorderSource*>& NewSourcesOut, FManifestSerializer* InManifestSerializer, TSet<FString>& UsedTrackNames)
 {
 
 	TArray<UTakeRecorderSource*> NewSources;
@@ -156,7 +157,25 @@ void UTakeRecorderSources::PreRecordingRecursive(TArray<UTakeRecorderSource*> In
 			// is using which sub-sequence so that we can push the correct sequence for all points of the Source's recording lifecycle.
 			if (Settings.bRecordSourcesIntoSubSequences && Source->SupportsSubscenes())
 			{
-				const FString& SubSequenceTrackName = ObjectTools::SanitizeObjectName(Source->GetSubsceneTrackName(InRootSequence));
+				const FString BaseSubSequenceTrackName = ObjectTools::SanitizeObjectName(Source->GetSubsceneTrackName(InRootSequence));
+				FString SubSequenceTrackName = BaseSubSequenceTrackName;
+				
+				// Simple duplicate name handling to prevent collisions. Normally the user should have a unique token in the title,
+				// such as {actor} which will keep the name unique.
+				{
+					int32 DuplicateCount = 1;
+					bool bTrackNameAlreadyUsed = false;
+					do
+					{
+						UsedTrackNames.Add(SubSequenceTrackName, &bTrackNameAlreadyUsed);
+						if (bTrackNameAlreadyUsed)
+						{
+							SubSequenceTrackName = FString::Printf(TEXT("%s_%s"), *BaseSubSequenceTrackName, *FString::FromInt(DuplicateCount++));
+						}
+					}
+					while (bTrackNameAlreadyUsed);
+				}
+				
 				const FString& SubSequenceAssetName = ObjectTools::SanitizeObjectName(Source->GetSubsceneAssetName(InRootSequence));
 
 				TargetSequence = CreateSubSequenceForSource(InRootSequence, SubSequenceTrackName, SubSequenceAssetName, Settings);
@@ -280,7 +299,7 @@ void UTakeRecorderSources::PreRecordingRecursive(TArray<UTakeRecorderSource*> In
 	{
 		// We don't want to nestle sub-sequences recursively so we always pass the Root Sequence and not the sequence
 		// created for a new source.
-		PreRecordingRecursive(NewSources, InRootSequence, NewSourcesOut, InManifestSerializer);
+		PreRecordingRecursive(NewSources, InRootSequence, NewSourcesOut, InManifestSerializer, UsedTrackNames);
 		SourcesSerialNumber++;
 	}
 }
@@ -294,7 +313,8 @@ void UTakeRecorderSources::StartRecordingPreRecordedSources(const FQualifiedFram
 void UTakeRecorderSources::PreRecordSources(TArray<UTakeRecorderSource *> InSources)
 {
 	PreRecordedSources.Reset(0);
-	PreRecordingRecursive(InSources, CachedLevelSequence, PreRecordedSources, CachedManifestSerializer);
+	TSet<FString> UsedTrackNames;
+	PreRecordingRecursive(InSources, CachedLevelSequence, PreRecordedSources, CachedManifestSerializer, UsedTrackNames);
 }
 
 
@@ -303,7 +323,8 @@ void UTakeRecorderSources::StartRecordingSource(TArray<UTakeRecorderSource *> In
 	// This calls PreRecording recursively on every source so that all sources that get added by another source
 	// have had PreRecording called.
 	TArray<UTakeRecorderSource *> NewSources;
-	PreRecordingRecursive(InSources, CachedLevelSequence, NewSources, CachedManifestSerializer);
+	TSet<FString> UsedTrackNames;
+	PreRecordingRecursive(InSources, CachedLevelSequence, NewSources, CachedManifestSerializer, UsedTrackNames);
 	if (NewSources.Num() > 0)
 	{
 		InSources.Append(NewSources);
diff --git a/Engine/Plugins/VirtualProduction/Takes/Source/TakesCore/Public/TakeRecorderSources.h b/Engine/Plugins/VirtualProduction/Takes/Source/TakesCore/Public/TakeRecorderSources.h
--- a/Engine/Plugins/VirtualProduction/Takes/Source/TakesCore/Public/TakeRecorderSources.h
+++ b/Engine/Plugins/VirtualProduction/Takes/Source/TakesCore/Public/TakeRecorderSources.h
@@ -208,7 +208,8 @@ private:
 	FOnSourcesChanged OnSourcesChangedEvent;
 
 	/** Calls PreRecording on sources recursively allowing them to create other sources which properly get PreRecording called on them as well. */
-	UE_API void PreRecordingRecursive(TArray<UTakeRecorderSource*> InSources, ULevelSequence* InRootSequence, TArray<UTakeRecorderSource*>& NewSourcesOut, FManifestSerializer* InManifestSerializer);
+	UE_API void PreRecordingRecursive(TArray<UTakeRecorderSource*> InSources, ULevelSequence* InRootSequence, TArray<UTakeRecorderSource*>& NewSourcesOut,
+		FManifestSerializer* InManifestSerializer, TSet<FString>& UsedTrackNames);
 
 	/** Finds the folder that the given Source should be created in, creating it if necessary. */
 	UE_API class UMovieSceneFolder* AddFolderForSource(const UTakeRecorderSource* InSource, class UMovieScene* InMovieScene);
