Fix how we auto track actors when using a world source.  The past implementation would fail to capture actors in a PIE environment and lead to incorrect take record.

#Synced-CL 49209914

diff --git a/Engine/Plugins/VirtualProduction/Takes/Source/TakeRecorderSources/Private/TakeRecorderWorldSource.cpp b/Engine/Plugins/VirtualProduction/Takes/Source/TakeRecorderSources/Private/TakeRecorderWorldSource.cpp
--- a/Engine/Plugins/VirtualProduction/Takes/Source/TakeRecorderSources/Private/TakeRecorderWorldSource.cpp
+++ b/Engine/Plugins/VirtualProduction/Takes/Source/TakeRecorderSources/Private/TakeRecorderWorldSource.cpp
@@ -24,6 +24,8 @@ UTakeRecorderWorldSourceSettings::UTakeRecorderWorldSourceSettings(const FObject
 	: Super(ObjInit)
 	, bRecordWorldSettings(true)
 	, bAutotrackActors(true)
+	, bAllowEditorOnlyActors(false)
+	, bOnlyRecordMoveableActors(true)
 {
 }
 
@@ -96,6 +98,13 @@ TArray<UTakeRecorderSource*> UTakeRecorderWorldSource::PreRecording(ULevelSequen
 	if (bAutotrackActors)
 	{
 		AutotrackActors(InSequence, World);
+		for (const TWeakObjectPtr<UTakeRecorderActorSource>& Source : AutotrackedSources)
+		{
+			if (Source.IsValid())
+			{
+				NewSources.Add(Source.Get());
+			}
+		}		
 	}
 
 	return NewSources;
@@ -109,6 +118,15 @@ TArray<UTakeRecorderSource*> UTakeRecorderWorldSource::PostRecording(class ULeve
 	{
 		SourcesToRemove.Add(WorldSource.Get());
 	}
+
+	for (const TWeakObjectPtr<UTakeRecorderActorSource>& Source : AutotrackedSources)
+	{
+		if (Source.IsValid())
+		{
+			SourcesToRemove.Add(Source.Get());
+		}
+	}
+
 	return SourcesToRemove;
 }
 
@@ -154,6 +172,11 @@ void UTakeRecorderWorldSource::AutotrackActors(class ULevelSequence* InSequence,
 	UTakeRecorderSources* Sources = InSequence->FindOrAddMetaData<UTakeRecorderSources>();
 
 	TArray<AActor*> ActorsBeingRecorded;
+	if (WorldSource.IsValid() && WorldSource.Get()->Target.IsValid())
+	{
+		ActorsBeingRecorded.Add(WorldSource.Get()->Target.Get());
+	}
+
 	if (Sources)
 	{
 		for (UTakeRecorderSource* Source : Sources->GetSources())
@@ -161,7 +184,7 @@ void UTakeRecorderWorldSource::AutotrackActors(class ULevelSequence* InSequence,
 			if (Source->IsA<UTakeRecorderActorSource>())
 			{
 				UTakeRecorderActorSource* ActorSource = Cast<UTakeRecorderActorSource>(Source);
-				if (ActorSource->Target.IsValid())
+				if (ActorSource && ActorSource->Target.IsValid())
 				{
 					ActorsBeingRecorded.Add(ActorSource->Target.Get());
 				}
@@ -169,46 +192,54 @@ void UTakeRecorderWorldSource::AutotrackActors(class ULevelSequence* InSequence,
 		}
 	}
 
-	TArray<TWeakObjectPtr<AActor>> ActorsToAdd;
+	AutotrackedSources.Reset();
 	for (ULevel* Level : InWorld->GetLevels())
 	{
 		if (Level)
 		{
 			for (AActor* Actor : Level->Actors)
 			{
-				if (!ActorsBeingRecorded.Contains(Actor))
+				if (Actor && !ActorsBeingRecorded.Contains(Actor))
 				{
-					ActorsToAdd.Add(Actor);
-				}
+					if (CanRecordActor(Actor))
+					{
+						UTakeRecorderActorSource* ActorSource = NewObject<UTakeRecorderActorSource>(Sources, UTakeRecorderActorSource::StaticClass(), NAME_None, RF_Transactional);
+						ActorSource->Target = Actor;
+
+						// Force the update to the recorded property map after change the actor source target. 
+						FPropertyChangedEvent PropertyChangedEvent(UTakeRecorderActorSource::StaticClass()->FindPropertyByName(GET_MEMBER_NAME_CHECKED(UTakeRecorderActorSource, Target)), EPropertyChangeType::ValueSet);
+						ActorSource->PostEditChangeProperty(PropertyChangedEvent);
+
+						// Add the source to our tracking table.
+						AutotrackedSources.Add(ActorSource);
+					}
+				}				
 			}
 		}
 	}
+}
 
-	TArray<FGuid> AddedGuids = SequencerPtr->AddActors(ActorsToAdd, false);
-
-	// Add to the "Autotracked Changes" folder
-	UMovieScene* MovieScene = InSequence->GetMovieScene();
-	FName AutotrackedFolderName = FName(TEXT("Autotracked Changes"));
-
-	UMovieSceneFolder* FolderToUse = nullptr;
-	for (UMovieSceneFolder* Folder : MovieScene->GetRootFolders())
+bool UTakeRecorderWorldSource::CanRecordActor(AActor* Actor) const
+{
+	if (!::IsValid(Actor))
 	{
-		if (Folder->GetFolderName() == AutotrackedFolderName)
-		{
-			FolderToUse = Folder;
-			break;
-		}
+		return false;
 	}
 
-	if (FolderToUse == nullptr)
+	// Skip editor-only actors
+	if (!bAllowEditorOnlyActors && Actor->IsEditorOnly())
 	{
-		FolderToUse = NewObject<UMovieSceneFolder>(MovieScene, NAME_None, RF_Transactional);
-		FolderToUse->SetFolderName(AutotrackedFolderName);
-		MovieScene->AddRootFolder(FolderToUse);
+		return false;
 	}
 
-	for (FGuid AddedGuid : AddedGuids)
+	// Skip actors with no root component or root component not movable
+	if (bOnlyRecordMoveableActors)
 	{
-		FolderToUse->AddChildObjectBinding(AddedGuid);
+		if (!Actor->GetRootComponent() || Actor->GetRootComponent()->Mobility == EComponentMobility::Static)
+		{
+			return false;
+		}
 	}
+
+	return true;
 }
diff --git a/Engine/Plugins/VirtualProduction/Takes/Source/TakeRecorderSources/Public/TakeRecorderWorldSource.h b/Engine/Plugins/VirtualProduction/Takes/Source/TakeRecorderSources/Public/TakeRecorderWorldSource.h
--- a/Engine/Plugins/VirtualProduction/Takes/Source/TakeRecorderSources/Public/TakeRecorderWorldSource.h
+++ b/Engine/Plugins/VirtualProduction/Takes/Source/TakeRecorderSources/Public/TakeRecorderWorldSource.h
@@ -32,6 +32,14 @@ public:
 	/** Add a binding and track for all actors that aren't explicitly being recorded */
 	UPROPERTY(config, EditAnywhere, BlueprintReadWrite, Category = "Source")
 	bool bAutotrackActors;
+
+	/** When autotracking actors to add, allow recording of editor only actor objects. */
+	UPROPERTY(config, EditAnywhere, BlueprintReadWrite, Category = "Source")
+	bool bAllowEditorOnlyActors;
+
+	/** When autotracking actors to add, only record actors that have a root component and are moveable. */
+	UPROPERTY(config, EditAnywhere, BlueprintReadWrite, Category = "Source")
+	bool bOnlyRecordMoveableActors;
 };
 
 
@@ -59,12 +67,18 @@ private:
 
 private:
 
-	/*
+	/**
 	 * Autotrack actors in the world that aren't already being recorded
 	 */
 	void AutotrackActors(class ULevelSequence* InSequence, UWorld* InWorld);
 
+	/**
+	 * If Autotracked actors is enabled this helper will determine if the actor can be recorded.
+	 */
+	bool CanRecordActor(AActor* Actor) const;
+
 	TWeakObjectPtr<UTakeRecorderActorSource> WorldSource;
+	TSet<TWeakObjectPtr<UTakeRecorderActorSource>> AutotrackedSources;
 };
 
 #undef UE_API
