Add a new cvar to control async saving in TakeRecorder.  There seems to be an issue on some machines where saving takes across a multi-user session will result in a check condition.  The theory is that async take saves are interacting with the the package saves in Multi-user events.  By default we disable async take saves to prevent interaction with the package save events but leave the cvar open to changing for experimentation.

#rnx
#Synced-CL 49551028

diff --git a/Engine/Plugins/VirtualProduction/Takes/Source/TakesCore/Private/TakesUtils.cpp b/Engine/Plugins/VirtualProduction/Takes/Source/TakesCore/Private/TakesUtils.cpp
--- a/Engine/Plugins/VirtualProduction/Takes/Source/TakesCore/Private/TakesUtils.cpp
+++ b/Engine/Plugins/VirtualProduction/Takes/Source/TakesCore/Private/TakesUtils.cpp
@@ -24,8 +24,11 @@ namespace TakesUtils
 {
 
 static bool bSaveTakesWithNoErrorFlag = false;
-static FAutoConsoleVariableRef  CVarSaveTakesWithNoError(TEXT("Takes.SaveTakesWithNoErrorFlag"), bSaveTakesWithNoErrorFlag, TEXT("For DrySquash only."));
+static FAutoConsoleVariableRef  CVarSaveTakesWithNoError(TEXT("Takes.SaveTakesWithNoErrorFlag"), bSaveTakesWithNoErrorFlag, TEXT("Enable the no error flag for package saves. By default, we always want to report errors back to the user."));
 
+static bool bSaveTakesAsync = false;
+static FAutoConsoleVariableRef  CVarSaveTakesAsyncSave(TEXT("Takes.SaveTakesWithAsyncSave"), bSaveTakesAsync, TEXT("Enable the async save flag for take recorder."));
+	
 class FArchiveClearExternalReferences : public FArchiveReplaceOrClearExternalReferences<UObject>
 {
 	using Super = FArchiveReplaceOrClearExternalReferences<UObject>;
@@ -184,7 +187,11 @@ void SaveAsset(UObject* InObject, bool bIsPossessable)
 
 	FSavePackageArgs SaveArgs;
 	SaveArgs.TopLevelFlags = RF_Standalone;
-	SaveArgs.SaveFlags = SAVE_Async;
+	if (bSaveTakesAsync)
+	{
+		SaveArgs.SaveFlags = SaveArgs.SaveFlags | SAVE_Async;		
+	}
+
 	if (bSaveTakesWithNoErrorFlag)
 	{
 		SaveArgs.SaveFlags = SaveArgs.SaveFlags | SAVE_NoError;
@@ -197,7 +204,10 @@ void SaveAsset(UObject* InObject, bool bIsPossessable)
 
 	if (UPackage::SavePackage(Package, NULL, *PackageFileName, SaveArgs))
 	{
-		UPackage::WaitForAsyncFileWrites();
+		if (bSaveTakesAsync)
+		{
+			UPackage::WaitForAsyncFileWrites();
+		}
 		double ElapsedTime = FPlatformTime::Seconds() - StartTime;
 		UE_LOG(LogTakesCore, Log, TEXT("Saved %s in %0.2f seconds"), *PackageName, ElapsedTime);
 	}
