TakeRecorder: Add GetPresetTargetRecordClass and SetPresetTargetRecordClass to take recorder subsystem.

#Synced-CL 50468325

diff --git a/Engine/Plugins/VirtualProduction/Takes/Source/TakeRecorder/Private/Recorder/TakeRecorderSubsystem.cpp b/Engine/Plugins/VirtualProduction/Takes/Source/TakeRecorder/Private/Recorder/TakeRecorderSubsystem.cpp
--- a/Engine/Plugins/VirtualProduction/Takes/Source/TakeRecorder/Private/Recorder/TakeRecorderSubsystem.cpp
+++ b/Engine/Plugins/VirtualProduction/Takes/Source/TakeRecorder/Private/Recorder/TakeRecorderSubsystem.cpp
@@ -275,6 +275,16 @@ bool UTakeRecorderSubsystem::HasPendingChanges() const
 	return Implementation->HasPendingChanges();
 }
 
+UClass* UTakeRecorderSubsystem::GetPresetTargetRecordClass() const
+{
+	return Implementation->GetPresetTargetRecordClass();
+}
+
+void UTakeRecorderSubsystem::SetPresetTargetRecordClass(const TSubclassOf<ULevelSequence> InClass)
+{
+	Implementation->SetPresetTargetRecordClass(InClass);
+}
+
 FOnTakeRecordingInitialized& UTakeRecorderSubsystem::GetOnRecordingPreInitializedEvent()
 {
 	return OnRecordingPreInitializeEvent;
diff --git a/Engine/Plugins/VirtualProduction/Takes/Source/TakeRecorder/Private/Recorder/TakeRecorderSubsystemImplementation.cpp b/Engine/Plugins/VirtualProduction/Takes/Source/TakeRecorder/Private/Recorder/TakeRecorderSubsystemImplementation.cpp
--- a/Engine/Plugins/VirtualProduction/Takes/Source/TakeRecorder/Private/Recorder/TakeRecorderSubsystemImplementation.cpp
+++ b/Engine/Plugins/VirtualProduction/Takes/Source/TakeRecorder/Private/Recorder/TakeRecorderSubsystemImplementation.cpp
@@ -829,6 +829,16 @@ bool UTakeRecorderSubsystemImplementation::HasPendingChanges() const
 	return bHasUserMadeChanges;
 }
 
+UClass* UTakeRecorderSubsystemImplementation::GetPresetTargetRecordClass() const
+{
+	return UTakePresetSettings::Get()->GetTargetRecordClass();
+}
+
+void UTakeRecorderSubsystemImplementation::SetPresetTargetRecordClass(const TSubclassOf<ULevelSequence> InClass)
+{
+	UTakePresetSettings::Get()->SetTargetRecordClass(InClass.Get());
+}
+
 UTakePreset* UTakeRecorderSubsystemImplementation::AllocateTransientPreset()
 {
 	return UTakePreset::AllocateTransientPreset(GetDefault<UTakeRecorderUserSettings>()->LastOpenedPreset.Get());
diff --git a/Engine/Plugins/VirtualProduction/Takes/Source/TakeRecorder/Private/Recorder/TakeRecorderSubsystemImplementation.h b/Engine/Plugins/VirtualProduction/Takes/Source/TakeRecorder/Private/Recorder/TakeRecorderSubsystemImplementation.h
--- a/Engine/Plugins/VirtualProduction/Takes/Source/TakeRecorder/Private/Recorder/TakeRecorderSubsystemImplementation.h
+++ b/Engine/Plugins/VirtualProduction/Takes/Source/TakeRecorder/Private/Recorder/TakeRecorderSubsystemImplementation.h
@@ -86,6 +86,8 @@ public:
 	virtual ETakeRecorderMode GetTakeRecorderMode() const override;
 	virtual UTakeRecorderNamingTokensData* GetNamingTokensData() const override;
 	virtual bool HasPendingChanges() const override;
+	virtual UClass* GetPresetTargetRecordClass() const override;
+	virtual void SetPresetTargetRecordClass(const TSubclassOf<ULevelSequence> InClass) override;
 	// ~End ITakeRecorderSubsystem interface
 	
 private:
diff --git a/Engine/Plugins/VirtualProduction/Takes/Source/TakeRecorder/Public/Recorder/ITakeRecorderSubsystemInterface.h b/Engine/Plugins/VirtualProduction/Takes/Source/TakeRecorder/Public/Recorder/ITakeRecorderSubsystemInterface.h
--- a/Engine/Plugins/VirtualProduction/Takes/Source/TakeRecorder/Public/Recorder/ITakeRecorderSubsystemInterface.h
+++ b/Engine/Plugins/VirtualProduction/Takes/Source/TakeRecorder/Public/Recorder/ITakeRecorderSubsystemInterface.h
@@ -79,4 +79,6 @@ public:
 	virtual ETakeRecorderMode GetTakeRecorderMode() const = 0;
 	virtual UTakeRecorderNamingTokensData* GetNamingTokensData() const = 0;
 	virtual bool HasPendingChanges() const = 0;
+	virtual UClass* GetPresetTargetRecordClass() const = 0;
+	virtual void SetPresetTargetRecordClass(const TSubclassOf<ULevelSequence> InClass) = 0;
 };
diff --git a/Engine/Plugins/VirtualProduction/Takes/Source/TakeRecorder/Public/Recorder/TakeRecorderSubsystem.h b/Engine/Plugins/VirtualProduction/Takes/Source/TakeRecorder/Public/Recorder/TakeRecorderSubsystem.h
--- a/Engine/Plugins/VirtualProduction/Takes/Source/TakeRecorder/Public/Recorder/TakeRecorderSubsystem.h
+++ b/Engine/Plugins/VirtualProduction/Takes/Source/TakeRecorder/Public/Recorder/TakeRecorderSubsystem.h
@@ -297,6 +297,14 @@ public:
 	/** @return Whether there are any pending changes, which ClearPendingTake would discard. */
 	UE_API virtual bool HasPendingChanges() const override;
 	
+	/** Retrieve the target record class of the Take Preset. */
+	UFUNCTION(BlueprintCallable, Category = TakeRecorder)
+	UE_API virtual UClass* GetPresetTargetRecordClass() const override;
+	
+	/** Set the target record class of the Take Preset. */
+	UFUNCTION(BlueprintCallable, Category = TakeRecorder)
+	UE_API virtual void SetPresetTargetRecordClass(const TSubclassOf<ULevelSequence> InClass) override;
+	
 	/**
 	 * Retrieve a multi-cast delegate that is triggered when a recording pre-initializes
 	 */
diff --git a/Engine/Plugins/VirtualProduction/Takes/Source/TakesCore/Private/TakePresetSettings.cpp b/Engine/Plugins/VirtualProduction/Takes/Source/TakesCore/Private/TakePresetSettings.cpp
--- a/Engine/Plugins/VirtualProduction/Takes/Source/TakesCore/Private/TakePresetSettings.cpp
+++ b/Engine/Plugins/VirtualProduction/Takes/Source/TakesCore/Private/TakePresetSettings.cpp
@@ -20,6 +20,18 @@ UClass* UTakePresetSettings::GetTargetRecordClass() const
 	return TargetRecordClass.TargetRecordClass ? TargetRecordClass.TargetRecordClass.Get() : ULevelSequence::StaticClass();
 }
 
+void UTakePresetSettings::SetTargetRecordClass(const TSoftClassPtr<ULevelSequence>& InClass)
+{
+	TargetRecordClass.TargetRecordClass = InClass;
+#if WITH_EDITOR
+	if (FProperty* Property = FindFProperty<FProperty>(GetClass(), GET_MEMBER_NAME_CHECKED(UTakePresetSettings, TargetRecordClass)))
+	{
+		FPropertyChangedEvent ChangedEvent(Property);
+		PostEditChangeProperty(ChangedEvent);
+	}
+#endif
+}
+
 #if WITH_EDITOR
 void UTakePresetSettings::PostEditChangeProperty(FPropertyChangedEvent& PropertyChangedEvent)
 {
diff --git a/Engine/Plugins/VirtualProduction/Takes/Source/TakesCore/Public/TakePresetSettings.h b/Engine/Plugins/VirtualProduction/Takes/Source/TakesCore/Public/TakePresetSettings.h
--- a/Engine/Plugins/VirtualProduction/Takes/Source/TakesCore/Public/TakePresetSettings.h
+++ b/Engine/Plugins/VirtualProduction/Takes/Source/TakesCore/Public/TakePresetSettings.h
@@ -2,8 +2,8 @@
 
 #pragma once
 
-#include "Templates/SubclassOf.h"
 #include "UObject/Object.h"
+#include "UObject/SoftObjectPtr.h"
 #include "TakePresetSettings.generated.h"
 
 #define UE_API TAKESCORE_API
@@ -39,6 +39,9 @@ public:
 
 	/** @return The class that recorded sequences should have */
 	UE_API UClass* GetTargetRecordClass() const;
+	
+	/** Sets the target record class */
+	UE_API void SetTargetRecordClass(const TSoftClassPtr<ULevelSequence>& InClass);
 
 	DECLARE_MULTICAST_DELEGATE(FOnSettingsChanged);
 	/** Invoked when any settings change. */
