TakeRecorder: Initial support recording soft actor and hard actor reference properties. This includes initial support for recording spawnable actor references.

Add ActorReferenceTrackRecorder for recording reference types, and converting between AActors and ActorReferenceKeys.
Recording possessable references should work in most situations, and recording spawnable references works in many situations. Recording over tracks with spawnables is not yet supported. Recording either to subsequences is not yet supported.

#Synced-CL 50653618

diff --git a/Engine/Plugins/VirtualProduction/Takes/Source/TakeRecorderSources/Private/TakeRecorderActorSource.cpp b/Engine/Plugins/VirtualProduction/Takes/Source/TakeRecorderSources/Private/TakeRecorderActorSource.cpp
--- a/Engine/Plugins/VirtualProduction/Takes/Source/TakeRecorderSources/Private/TakeRecorderActorSource.cpp
+++ b/Engine/Plugins/VirtualProduction/Takes/Source/TakeRecorderSources/Private/TakeRecorderActorSource.cpp
@@ -249,8 +249,26 @@ TArray<UTakeRecorderSource*> UTakeRecorderActorSource::PreRecording(ULevelSequen
 	// Look to see if the movie scene already has our object binding in it (which is common if we're recording a new take) so we can
 	// replace the data that was already there.
 	CachedObjectBindingGuid = ResolveActorFromSequence(ActorToRecord, TargetLevelSequence);
-	CleanExistingDataFromSequence(CachedObjectBindingGuid, *TargetLevelSequence);
 	
+	// Check if the binding was just created by an actor reference track that was processed by another actor source, prior to this one.
+	// In this case the existing binding is the one we want to use.
+	UTakeRecorderSources* SourcesList = GetTypedOuter<UTakeRecorderSources>();
+	check(SourcesList);
+	const FGuid* ExistingBindingGuid = SourcesList->GetNewActorBindings().Find(ActorToRecord);
+	const bool bIsExistingBindingValid = ExistingBindingGuid != nullptr;
+	if (bIsExistingBindingValid)
+	{
+		if (!CachedObjectBindingGuid.IsValid() )
+		{
+			CachedObjectBindingGuid = *ExistingBindingGuid;
+		}
+		ensure(CachedObjectBindingGuid == *ExistingBindingGuid);
+	}
+	else
+	{
+		CleanExistingDataFromSequence(CachedObjectBindingGuid, *TargetLevelSequence);
+	}
+		
 	// Initialize the header for this actor in the Manifest Serializer for streaming data capture.
 	FName SerializedType("Actor");
 	FActorFileHeader Header(ObjectBindingName, ActorToRecord->GetActorLabel(), SerializedType, ActorToRecord->GetClass()->GetPathName(), false);
@@ -266,9 +284,12 @@ TArray<UTakeRecorderSource*> UTakeRecorderActorSource::PreRecording(ULevelSequen
 		}
 		else
 		{
-			// Create a Possessable object binding in the Sequence and then bind it to our actor
-			CachedObjectBindingGuid = MovieScene->AddPossessable(ActorToRecord->GetActorLabel(), ActorToRecord->GetClass());
-			TargetLevelSequence->BindPossessableObject(CachedObjectBindingGuid, *ActorToRecord, ActorToRecord->GetWorld());
+			if (!bIsExistingBindingValid)
+			{
+				// Create a Possessable object binding in the Sequence and then bind it to our actor
+				CachedObjectBindingGuid = MovieScene->AddPossessable(ActorToRecord->GetActorLabel(), ActorToRecord->GetClass());
+				TargetLevelSequence->BindPossessableObject(CachedObjectBindingGuid, *ActorToRecord, ActorToRecord->GetWorld());
+			}
 			Header.bRecordToPossessable = true;
 		}
 	}
@@ -281,8 +302,8 @@ TArray<UTakeRecorderSource*> UTakeRecorderActorSource::PreRecording(ULevelSequen
 		CreateBindingParams.bSpawnable = true;
 		CreateBindingParams.bAllowCustomBinding = true;
 		CreateBindingParams.BindingNameOverride = ActorToRecord->GetActorLabel();
-		CachedObjectBindingGuid = FSequencerUtilities::CreateOrReplaceBinding(nullptr, InSequence, ActorToRecord, CreateBindingParams);
-		CachedObjectTemplate = CastChecked<AActor>(MovieSceneHelpers::GetObjectTemplate(InSequence, CachedObjectBindingGuid, MovieSceneHelpers::CreateTransientSharedPlaybackState(ActorToRecord, InSequence)));
+		CachedObjectBindingGuid = ExistingBindingGuid ? *ExistingBindingGuid : FSequencerUtilities::CreateOrReplaceBinding(nullptr, InSequence, ActorToRecord, CreateBindingParams);
+		CachedObjectTemplate = Cast<AActor>(MovieSceneHelpers::GetObjectTemplate(InSequence, CachedObjectBindingGuid, MovieSceneHelpers::CreateTransientSharedPlaybackState(ActorToRecord, InSequence)));
 		
 		if (CachedObjectTemplate.IsValid())
 		{
@@ -476,6 +497,7 @@ void UTakeRecorderActorSource::CreateSectionRecordersRecursive(UObject* ObjectTo
 		}
 
 		TArray<IMovieSceneTrackRecorderFactory*> ModularFactories = IModularFeatures::Get().GetModularFeatureImplementations<IMovieSceneTrackRecorderFactory>(MovieSceneSectionRecorderFactoryName);
+		TArray<IMovieSceneTrackRecorderFactory*> PropertyInstanceFactories;
 		bool bFoundRecorder = false;
 		for (IMovieSceneTrackRecorderFactory* Factory : ModularFactories)
 		{
@@ -500,6 +522,11 @@ void UTakeRecorderActorSource::CreateSectionRecordersRecursive(UObject* ObjectTo
 					break;
 				}
 			}
+			
+			if (Factory->CanHandlePropertyInstance())
+			{
+				PropertyInstanceFactories.Add(Factory);
+			}
 		}
 
 		if (!bFoundRecorder)
@@ -526,11 +553,9 @@ void UTakeRecorderActorSource::CreateSectionRecordersRecursive(UObject* ObjectTo
 
 			if (PropertyInstance)
 			{
-				FMovieScenePropertyTrackRecorderFactory TrackRecorderFactory;
-				if (TrackRecorderFactory.CanRecordProperty(ObjectToRecord, PropertyInstance))
+				auto CreateTrackRecorderForProperty = [&](const IMovieSceneTrackRecorderFactory& TrackRecorderFactory)
 				{
-					UMovieSceneTrackRecorder* SectionRecorder = TrackRecorderFactory.CreateTrackRecorderForProperty(ObjectToRecord, Property.PropertyName);
-					if (SectionRecorder)
+					if (UMovieSceneTrackRecorder* SectionRecorder = TrackRecorderFactory.CreateTrackRecorderForProperty(ObjectToRecord, Property.PropertyName))
 					{
 						TrackRecorders.Add(SectionRecorder);
 						if (TrackRecorderFactory.IsSerializable())
@@ -545,6 +570,28 @@ void UTakeRecorderActorSource::CreateSectionRecordersRecursive(UObject* ObjectTo
 						}
 						SectionRecorder->CreateTrack(this, ObjectToRecord, TargetLevelSequence->GetMovieScene(), GetSettingsObjectForFactory(TrackRecorderFactory.GetSettingsClass()), Guid);
 					}
+				};
+				
+				auto FindPropertyInstanceFactoryLambda = [ObjectToRecord, PropertyInstance](const IMovieSceneTrackRecorderFactory* Factory)
+				{
+					return Factory->CanRecordProperty(ObjectToRecord, PropertyInstance);
+				};
+				
+				// First check if there's a dedicated recorder to handle a property instance.
+				if (IMovieSceneTrackRecorderFactory** PropertyInstanceFactoryPtr = PropertyInstanceFactories.FindByPredicate(FindPropertyInstanceFactoryLambda))
+				{
+					IMovieSceneTrackRecorderFactory* PropertyInstanceFactory = *PropertyInstanceFactoryPtr;
+					check(PropertyInstanceFactory->CanHandlePropertyInstance());
+					CreateTrackRecorderForProperty(*PropertyInstanceFactory);
+				}
+				else
+				{
+					// Fall back to the generic property recorder factory.
+					FMovieScenePropertyTrackRecorderFactory TrackRecorderFactory;
+					if (TrackRecorderFactory.CanRecordProperty(ObjectToRecord, PropertyInstance))
+					{
+						CreateTrackRecorderForProperty(TrackRecorderFactory);
+					}
 				}
 			}
 			else 
@@ -842,6 +889,12 @@ TArray<UTakeRecorderSource*> UTakeRecorderActorSource::PostRecording(ULevelSeque
 	CachedObjectTemplate = nullptr;
 	CachedComponentList.Empty();
 
+	if (bIsTemporarySource)
+	{
+		// So this source can be cleaned up later.
+		AddedActorSources.Add(this);
+	}
+	
 	// We may have generated some temporary recording sources
 	return AddedActorSources;
 }
@@ -1214,7 +1267,10 @@ void UTakeRecorderActorSource::RebuildRecordedPropertyMapRecursive(const FFieldV
 			// shown there do actually have something trying to record them).
 			for (IMovieSceneTrackRecorderFactory* Factory : ModularFactories)
 			{
-				if (InObject.IsUObject() && Factory->CanRecordProperty(InObject.ToUObject(), *It))
+				if ((InObject.IsUObject() && Factory->CanRecordProperty(InObject.ToUObject(), *It))
+					||
+					// A factory may have special handling for any property of this type.
+					(Factory->CanHandlePropertyInstance() && Factory->CanRecordProperty(Target.Get(), *It)))
 				{
 			 		DebugDisplayName = Factory->GetDisplayName();
 
@@ -1390,22 +1446,7 @@ FText UTakeRecorderActorSource::GetDescriptionTextImpl() const
 
 FGuid UTakeRecorderActorSource::ResolveActorFromSequence(AActor* InActor, ULevelSequence* CurrentSequence) const
 {
-	UMovieScene* MovieScene = CurrentSequence->GetMovieScene();
-
-	// Look through all Possessables in the sequence to see if there's one with the same name as the actor. We purposely do not look at 
-	// Spawnables so that recording as spawnable will always create a new spawnable.
-	for (int32 PossessableCount = 0; PossessableCount < MovieScene->GetPossessableCount(); ++PossessableCount)
-	{
-		const FMovieScenePossessable& Possessable = MovieScene->GetPossessable(PossessableCount);
-		if (Possessable.GetName() == InActor->GetActorLabel())
-		{
-			return Possessable.GetGuid();
-		}
-	}
-
-	// There's no Possessable with the same name as the actor, so this actor hasn't been added
-	// to the sequence yet.
-	return FGuid();
+	return TakesUtils::ResolveActorFromSequence(InActor, CurrentSequence);
 }
 
 void UTakeRecorderActorSource::PostProcessCreatedObjectTemplateImpl(AActor* ObjectTemplate)
@@ -1674,6 +1715,10 @@ void UTakeRecorderActorSource::CreateNewActorSourceForReferencedActors()
 			}
 		}
 	}
+	
+	TArray<TObjectPtr<AActor>> NewActorBindingsArray;
+	SourcesList->GetNewActorBindings().GenerateKeyArray(NewActorBindingsArray);
+	NewReferencedActors.Append(MoveTemp(NewActorBindingsArray));
 
 	for (AActor* Actor : NewReferencedActors)
 	{
@@ -1804,6 +1849,16 @@ bool UTakeRecorderActorSource::GetRecordToPossessable() const
 	return RecordType == ETakeRecorderActorRecordType::Possessable;
 }
 
+void UTakeRecorderActorSource::SetIsTemporarySource(bool bValue)
+{
+	bIsTemporarySource = bValue;
+}
+
+bool UTakeRecorderActorSource::IsOtherActorPossessable(AActor* InActor) const
+{
+	return TakeRecorderSourcesUtils::IsActorPossessable(this, InActor);
+}
+
 void UTakeRecorderActorSource::InitializeFactorySettingsObject(TSubclassOf<UMovieSceneTrackRecorderSettings> InClass)
 {
 	ensure(InClass);
diff --git a/Engine/Plugins/VirtualProduction/Takes/Source/TakeRecorderSources/Private/TakeRecorderNearbySpawnedActorSource.cpp b/Engine/Plugins/VirtualProduction/Takes/Source/TakeRecorderSources/Private/TakeRecorderNearbySpawnedActorSource.cpp
--- a/Engine/Plugins/VirtualProduction/Takes/Source/TakeRecorderSources/Private/TakeRecorderNearbySpawnedActorSource.cpp
+++ b/Engine/Plugins/VirtualProduction/Takes/Source/TakeRecorderSources/Private/TakeRecorderNearbySpawnedActorSource.cpp
@@ -187,21 +187,10 @@ void UTakeRecorderNearbySpawnedActorSource::HandleActorSpawned(AActor* Actor, cl
 
 	UE_LOG(LogTakesCore, Log, TEXT("Actor: %s IsValid: %d Unreachable: %d PendingKillPending: %d"), *Actor->GetName(), ::IsValid(Actor), Actor->IsUnreachable(), Actor->IsPendingKillPending());
 
-	UTakeRecorderSources* Sources = InSequence->FindOrAddMetaData<UTakeRecorderSources>();
-
-	UTakeRecorderActorSource* ActorSource = Sources->AddSource<UTakeRecorderActorSource>();
-	ActorSource->Target = Actor;
-
-	// Send a PropertyChangedEvent so the class catches the callback and rebuilds the property map. We can't rely on the Actor rebuilding the map on PreRecording
-	// because that would wipe out any user adjustments from one added natively.
-	FPropertyChangedEvent PropertyChangedEvent(UTakeRecorderActorSource::StaticClass()->FindPropertyByName(GET_MEMBER_NAME_CHECKED(UTakeRecorderActorSource, Target)), EPropertyChangeType::ValueSet);
-	ActorSource->PostEditChangeProperty(PropertyChangedEvent);
-
-	// This has to be called after setting the Target and propegating the change event so that it has a chance to know what to record
-	// about the actor.
-	Sources->StartRecordingSource(TArray<UTakeRecorderSource*>({ ActorSource }), Sources->GetCachedFrameTime());
-
-	SpawnedActorSources.Add(ActorSource);
+	if (UTakeRecorderActorSource* ActorSource = TakeRecorderSourcesUtils::StartRecordingActorDuringSession(Actor, InSequence))
+	{
+		SpawnedActorSources.Add(ActorSource);
+	}
 }
 
 #undef LOCTEXT_NAMESPACE // "UTakeRecorderNearbySpawnedActorSource"
diff --git a/Engine/Plugins/VirtualProduction/Takes/Source/TakeRecorderSources/Private/TakeRecorderSourcesModule.cpp b/Engine/Plugins/VirtualProduction/Takes/Source/TakeRecorderSources/Private/TakeRecorderSourcesModule.cpp
--- a/Engine/Plugins/VirtualProduction/Takes/Source/TakeRecorderSources/Private/TakeRecorderSourcesModule.cpp
+++ b/Engine/Plugins/VirtualProduction/Takes/Source/TakeRecorderSources/Private/TakeRecorderSourcesModule.cpp
@@ -165,6 +165,8 @@ void FTakeRecorderSourcesModule::StartupModule()
 
 	ISequencerModule& SequencerModule = FModuleManager::Get().LoadModuleChecked<ISequencerModule>("Sequencer");
 	OnSequencerCreatedHandle = SequencerModule.RegisterOnSequencerCreated(FOnSequencerCreated::FDelegate::CreateRaw(this, &FTakeRecorderSourcesModule::OnSequencerCreated));
+	
+	UTakeRecorderSources::GetOnActorBindingAddedDelegate().AddRaw(this, &FTakeRecorderSourcesModule::OnNewActorReferencingBindingAdded);
 }
 
 void FTakeRecorderSourcesModule::ShutdownModule()
@@ -187,6 +189,8 @@ void FTakeRecorderSourcesModule::ShutdownModule()
 	FTakeRecorderSourcesCommands::Unregister();
 
 	UnregisterMenuExtensions();
+	
+	UTakeRecorderSources::GetOnActorBindingAddedDelegate().RemoveAll(this);
 }
 
 void FTakeRecorderSourcesModule::RegisterCanRecordDelegate(FName HandleId, FCanRecordDelegate InDelegate)
@@ -772,7 +776,15 @@ void FTakeRecorderSourcesModule::OnSequencerCreated(TSharedRef<ISequencer> Seque
 		}
 	});
 }
-	
+
+void FTakeRecorderSourcesModule::OnNewActorReferencingBindingAdded(AActor* InActor, ULevelSequence* InLevelSequence)
+{	
+	if (UTakeRecorderActorSource* ActorSource = TakeRecorderSourcesUtils::StartRecordingActorDuringSession(InActor, InLevelSequence))
+	{
+		ActorSource->SetIsTemporarySource(true);
+	}
+}
+
 IMPLEMENT_MODULE(FTakeRecorderSourcesModule, TakeRecorderSources);
 }
 
diff --git a/Engine/Plugins/VirtualProduction/Takes/Source/TakeRecorderSources/Private/TakeRecorderSourcesModule.h b/Engine/Plugins/VirtualProduction/Takes/Source/TakeRecorderSources/Private/TakeRecorderSourcesModule.h
--- a/Engine/Plugins/VirtualProduction/Takes/Source/TakeRecorderSources/Private/TakeRecorderSourcesModule.h
+++ b/Engine/Plugins/VirtualProduction/Takes/Source/TakeRecorderSources/Private/TakeRecorderSourcesModule.h
@@ -68,5 +68,7 @@ private:
 	void RecordSelectedActors();
 
 	void OnSequencerCreated(TSharedRef<ISequencer> Sequencer);
+	
+	void OnNewActorReferencingBindingAdded(AActor* InActor, ULevelSequence* InLevelSequence);
 };
 }
diff --git a/Engine/Plugins/VirtualProduction/Takes/Source/TakeRecorderSources/Private/TakeRecorderSourcesUtils.h b/Engine/Plugins/VirtualProduction/Takes/Source/TakeRecorderSources/Private/TakeRecorderSourcesUtils.h
--- a/Engine/Plugins/VirtualProduction/Takes/Source/TakeRecorderSources/Private/TakeRecorderSourcesUtils.h
+++ b/Engine/Plugins/VirtualProduction/Takes/Source/TakeRecorderSources/Private/TakeRecorderSourcesUtils.h
@@ -92,6 +92,38 @@ namespace TakeRecorderSourcesUtils
 
 		return false;
 	}
+	
+	/** If an actor is configured to be possessable. */
+	static bool IsActorPossessable(const UTakeRecorderActorSource* InSource, AActor* Actor)
+	{
+		check(InSource);
+		
+		if (InSource->Target == Actor && InSource->GetRecordToPossessable())
+		{
+			return true;
+		}
+
+		const UTakeRecorderSources* OwningSources = CastChecked<UTakeRecorderSources>(InSource->GetOuter());
+		for (UTakeRecorderSource* Source : OwningSources->GetSources())
+		{
+			if (const UTakeRecorderActorSource* ActorSource = Cast<UTakeRecorderActorSource>(Source))
+			{
+				if (ActorSource->bEnabled && ActorSource->Target.Get() == Actor && ActorSource->GetRecordToPossessable())
+				{
+					return true;
+				}
+			}
+		}
+
+		// Fallback to global settings in the event the actor doesn't have a source for it. This is what will be used
+		// to create the actor source in the future anyway, if needed.
+		if (const UTakeRecorderSources* Sources = InSource->GetTypedOuter<UTakeRecorderSources>())
+		{
+			return Sources->GetSettings().bRecordToPossessable;
+		}
+		
+		return false;
+	}
 
 	static bool IsActorRecordable(const AActor* Actor)
 	{
@@ -182,4 +214,28 @@ namespace TakeRecorderSourcesUtils
 
 		return MovieSceneSequenceID::Root;
 	}
+	
+	/** Create a new actor source for an actor which can be done during an active recording session. */
+	static UTakeRecorderActorSource* StartRecordingActorDuringSession(AActor* InActor, ULevelSequence* InLevelSequence)
+	{
+		if (!IsValid(InActor) || !IsValid(InLevelSequence))
+		{
+			return nullptr;
+		}
+		
+		UTakeRecorderSources* Sources = InLevelSequence->FindOrAddMetaData<UTakeRecorderSources>();
+
+		UTakeRecorderActorSource* ActorSource = Sources->AddSource<UTakeRecorderActorSource>();
+		ActorSource->Target = InActor;
+
+		// Send a PropertyChangedEvent so the class catches the callback and rebuilds the property map. We can't rely on the Actor rebuilding the map on PreRecording
+		// because that would wipe out any user adjustments from one added natively.
+		FPropertyChangedEvent PropertyChangedEvent(UTakeRecorderActorSource::StaticClass()->FindPropertyByName(GET_MEMBER_NAME_CHECKED(UTakeRecorderActorSource, Target)), EPropertyChangeType::ValueSet);
+		ActorSource->PostEditChangeProperty(PropertyChangedEvent);
+
+		// This has to be called after setting the Target and propegating the change event so that it has a chance to know what to record
+		// about the actor.
+		Sources->StartRecordingSource(TArray<UTakeRecorderSource*>({ ActorSource }), Sources->GetCachedFrameTime());
+		return ActorSource;
+	}
 }
\ No newline at end of file
diff --git a/Engine/Plugins/VirtualProduction/Takes/Source/TakeRecorderSources/Public/TakeRecorderActorSource.h b/Engine/Plugins/VirtualProduction/Takes/Source/TakeRecorderSources/Public/TakeRecorderActorSource.h
--- a/Engine/Plugins/VirtualProduction/Takes/Source/TakeRecorderSources/Public/TakeRecorderActorSource.h
+++ b/Engine/Plugins/VirtualProduction/Takes/Source/TakeRecorderSources/Public/TakeRecorderActorSource.h
@@ -173,6 +173,9 @@ public:
 	/** Get the record type. If set to project default, gets the type from the project settings */
 	UE_API bool GetRecordToPossessable() const;
 
+	/** If this source should be cleaned up later. */
+	void SetIsTemporarySource(bool bValue);
+	
 protected:
 
 	/**
@@ -274,6 +277,8 @@ protected:
 	UE_API FTransform GetRecordedActorAnimationInitialRootTransform(class AActor* OtherActor) const override;
 	/** Returns the root level sequence being recorded into */
 	ULevelSequence* GetRootLevelSequence() const override { return RootLevelSequence; };
+	/** Returns if an actor is possesable. */
+	UE_API virtual bool IsOtherActorPossessable(AActor* InActor) const override;
 	// ~IMovieSceneTrackRecorderHost Interface
 	
 	/** Initializes an instance of the specified class if we don't already have it in our Settings array. */
@@ -336,7 +341,9 @@ private:
 	*  Serializer
 	*/
 	FActorSerializer ActorSerializer;
-
+	
+	/** If this actor source should be considered temporary, and removed later. */
+	bool bIsTemporarySource = false;
 };
 
 #undef UE_API
\ No newline at end of file
diff --git a/Engine/Plugins/VirtualProduction/Takes/Source/TakeTrackRecorders/Private/MovieSceneActorReferenceTrackRecorder.cpp b/Engine/Plugins/VirtualProduction/Takes/Source/TakeTrackRecorders/Private/MovieSceneActorReferenceTrackRecorder.cpp
new file mode 100644
--- /dev/null
+++ b/Engine/Plugins/VirtualProduction/Takes/Source/TakeTrackRecorders/Private/MovieSceneActorReferenceTrackRecorder.cpp
@@ -0,0 +1,326 @@
+// Copyright Epic Games, Inc. All Rights Reserved.
+
+#include "TrackRecorders/MovieSceneActorReferenceTrackRecorder.h"
+#include "Channels/MovieSceneChannelProxy.h"
+#include "Components/SceneComponent.h"
+#include "GameFramework/Actor.h"
+#include "LevelSequence.h"
+#include "MovieScene.h"
+#include "MovieSceneSpawnableAnnotation.h"
+#include "SequencerUtilities.h"
+#include "TakeRecorderSource.h"
+#include "TakeRecorderSources.h"
+#include "TakesCoreLog.h"
+#include "TakesUtils.h"
+#include "Tracks/MovieSceneActorReferenceTrack.h"
+
+#include UE_INLINE_GENERATED_CPP_BY_NAME(MovieSceneActorReferenceTrackRecorder)
+
+bool FMovieSceneActorReferenceTrackRecorderFactory::CanRecordObject(UObject* InObjectToRecord) const
+{
+	return false;
+}
+
+UMovieSceneTrackRecorder* FMovieSceneActorReferenceTrackRecorderFactory::CreateTrackRecorderForObject() const
+{
+	return nullptr;
+}
+
+bool FMovieSceneActorReferenceTrackRecorderFactory::CanRecordProperty(UObject* InObjectToRecord, FProperty* InPropertyToRecord) const
+{
+	if (InPropertyToRecord && (InPropertyToRecord->IsA<FSoftObjectProperty>() || InPropertyToRecord->IsA<FObjectProperty>()))
+	{
+		return true;
+	}
+	return false;
+}
+
+UMovieSceneTrackRecorder* FMovieSceneActorReferenceTrackRecorderFactory::CreateTrackRecorderForProperty(UObject* InObjectToRecord,
+	const FName& InPropertyToRecord) const
+{
+	UMovieSceneActorReferenceTrackRecorder* TrackRecorder = NewObject<UMovieSceneActorReferenceTrackRecorder>();
+	TrackRecorder->CreateBinding(InPropertyToRecord);
+	return TrackRecorder;
+}
+
+void UMovieSceneActorReferenceTrackRecorder::CreateBinding(const FName& InPropertyToRecord)
+{
+	PropertyBinding = MakeShared<FTrackInstancePropertyBindings>(InPropertyToRecord, InPropertyToRecord.ToString());
+}
+
+void UMovieSceneActorReferenceTrackRecorder::CreateTrackImpl()
+{
+	bSetFirstKey = true;
+
+	check(PropertyBinding.IsValid())
+	check(MovieScene.IsValid());
+	check(ObjectGuid.IsValid());
+	
+	const TStrongObjectPtr<UMovieScene> MovieScenePin = MovieScene.Pin();
+	
+	const FName TrackName = *PropertyBinding->GetPropertyPath();
+	UMovieSceneActorReferenceTrack* Track = MovieScenePin->FindTrack<UMovieSceneActorReferenceTrack>(ObjectGuid, TrackName);
+	if (!Track)
+	{
+		Track = MovieScenePin->AddTrack<UMovieSceneActorReferenceTrack>(ObjectGuid);
+	}
+	else
+	{
+		Track->RemoveAllAnimationData();
+	}
+	if (!Track)
+	{
+		Track = MovieScenePin->AddTrack<UMovieSceneActorReferenceTrack>(ObjectGuid);
+	}
+	else
+	{
+		Track->RemoveAllAnimationData();
+	}
+	if (Track)
+	{
+		Track->SetPropertyNameAndPath(PropertyBinding->GetPropertyName(), PropertyBinding->GetPropertyPath());
+		MovieSceneSection = Cast<UMovieSceneActorReferenceSection>(Track->CreateNewSection());
+		// We only set the track defaults when we're not loading from a serialized recording. Serialized recordings don't store channel defaults but will always store a
+		// key on the first frame which will accomplish the same.
+		
+		AActor* ActorValue = EvaluateBinding();
+		LastBindingValue = ActorValue;
+		
+		const FMovieSceneActorReferenceKey ReferenceKey = ActorToReferenceKey(ActorValue);
+		FMovieSceneActorReferenceChannel* ActorReferenceChannel = MovieSceneSection->GetChannelProxy().GetChannel<FMovieSceneActorReferenceChannel>(0);
+		if (ensure(ActorReferenceChannel))
+		{
+			ActorReferenceChannel->SetDefault(ReferenceKey);
+		}
+		
+		Track->AddSection(*MovieSceneSection);
+	}
+}
+
+void UMovieSceneActorReferenceTrackRecorder::RecordSampleImpl(const FQualifiedFrameTime& CurrentTime)
+{
+	if (!MovieSceneSection.IsValid())
+	{
+		return;
+	}
+	
+	bIsRecording = true;
+
+	if (ObjectToRecord.IsValid())
+	{
+		if (bSetFirstKey)
+		{
+			bSetFirstKey = false;
+			FMovieSceneActorReferenceChannel* Channel = MovieSceneSection->GetChannelProxy().GetChannel<FMovieSceneActorReferenceChannel>(0);
+			if (ensure(Channel))
+			{
+				if (MovieSceneSection->HasStartFrame())
+				{
+					const FMovieSceneActorReferenceKey ReferenceKey = ActorToReferenceKey(LastBindingValue.Get());
+					Channel->GetData().AddKey(MovieSceneSection->GetInclusiveStartFrame(), ReferenceKey);
+				}
+			}
+		}
+
+		const FFrameRate   TickResolution = MovieSceneSection->GetTypedOuter<UMovieScene>()->GetTickResolution();
+		const FFrameNumber CurrentFrame = CurrentTime.ConvertTo(TickResolution).FloorToFrame();
+
+		MovieSceneSection->SetEndFrame(CurrentFrame);
+		
+		AActor* ActorValue = EvaluateBinding();
+		if (LastBindingValue.Get() != ActorValue)
+		{
+			LastBindingValue = ActorValue;
+			FMovieSceneActorReferenceChannel* Channel = MovieSceneSection->GetChannelProxy().GetChannel<FMovieSceneActorReferenceChannel>(0);
+			if (ensure(Channel))
+			{
+				const FMovieSceneActorReferenceKey ReferenceKey = ActorToReferenceKey(ActorValue);
+				Channel->GetData().AddKey(CurrentFrame, ReferenceKey);
+			}
+		}
+	}
+}
+
+
+void UMovieSceneActorReferenceTrackRecorder::FinalizeTrackImpl()
+{
+	bIsRecording = false;
+	
+	if (!MovieSceneSection.IsValid())
+	{
+		return;
+	}
+
+	if (!ObjectToRecord.IsValid())
+	{
+		return;
+	}
+
+	const FTrackRecorderSettings TrackRecorderSettings = OwningTakeRecorderSource->GetTrackRecorderSettings();
+
+	if (TrackRecorderSettings.bRemoveRedundantTracks)
+	{
+		RemoveRedundantTracks();
+	}
+}
+
+void UMovieSceneActorReferenceTrackRecorder::RemoveRedundantTracks()
+{
+		if (!ObjectToRecord.IsValid() || !MovieSceneSection.IsValid() || !PropertyBinding.IsValid() || !MovieScene.IsValid())
+		{
+			return;
+		}
+		
+		FTrackRecorderSettings TrackRecorderSettings = OwningTakeRecorderSource->GetTrackRecorderSettings();
+			
+		// The section can be removed if this is a spawnable since the spawnable template should have the same default values
+		bool bRemoveSection = true;
+
+		bool bExclude = FTrackRecorderSettings::IsExcludePropertyTrack(ObjectToRecord.Get(), PropertyBinding->GetPropertyPath(), TrackRecorderSettings.DefaultTracks);
+		if (!bExclude)
+		{
+			// If any channel has more than 1 key, the track cannot be removed
+			for (const FMovieSceneChannelEntry& Entry : MovieSceneSection->GetChannelProxy().GetAllEntries())
+			{
+				TArrayView<FMovieSceneChannel* const> Channels = Entry.GetChannels();
+
+				for (int32 Index = 0; Index < Channels.Num(); ++Index)
+				{
+					if (Channels[Index]->GetNumKeys() > 1)
+					{
+						return;
+					}
+				}
+			}
+
+			// Assumes each channel is left with 1 or no keys, so the keys can be removed and the default value set
+			FMovieSceneActorReferenceChannel* Channel = MovieSceneSection->GetChannelProxy().GetChannel<FMovieSceneActorReferenceChannel>(0);
+			const FMovieSceneActorReferenceKey DefaultReferenceKey = Channel->GetDefault();
+
+			// Reset channels
+			for (const FMovieSceneChannelEntry& Entry : MovieSceneSection->GetChannelProxy().GetAllEntries())
+			{
+				TArrayView<FMovieSceneChannel* const> Channels = Entry.GetChannels();
+
+				for (int32 Index = 0; Index < Channels.Num(); ++Index)
+				{
+					Channels[Index]->Reset();
+				}
+			}
+
+			Channel->SetDefault(DefaultReferenceKey);
+
+			// This section can only be removed if the CDO value is the same and it's not on the allow list of default property tracks.
+			// NOTE: On the property recorder we restrict this just to TrackRecorderSettings.bRecordToPossessable, but with actor references this is
+			// necessary in the event we start recording a spawnable as a default value, before starting the recording (a common scenario).
+			{
+				bRemoveSection = false;
+
+				if (UObject* DefaultObject = ObjectToRecord->GetClass()->GetDefaultObject())
+				{
+					AActor* CurrentActorValue = EvaluateBinding(DefaultObject);
+					const FMovieSceneActorReferenceKey ActorReferenceKey = ActorToReferenceKey(CurrentActorValue);
+					if (ActorReferenceKey == DefaultReferenceKey)
+					{
+						bRemoveSection = true;
+					}
+				}
+
+				if (bRemoveSection && FTrackRecorderSettings::IsDefaultPropertyTrack(ObjectToRecord.Get(), PropertyBinding->GetPropertyPath(), TrackRecorderSettings.DefaultTracks))
+				{
+					bRemoveSection = false;
+				}
+			}
+		}
+
+		if (bRemoveSection || bExclude)
+		{
+			UMovieSceneTrack* MovieSceneTrack = CastChecked<UMovieSceneTrack>(MovieSceneSection->GetOuter());
+
+			UE_LOG(LogTakesCore, Log, TEXT("Removed unused track (%s) for (%s)"), *MovieSceneTrack->GetTrackName().ToString(), *ObjectToRecord->GetName());
+
+			MovieSceneTrack->RemoveSection(*MovieSceneSection);
+			MovieScene->RemoveTrack(*MovieSceneTrack);
+		}
+}
+
+FMovieSceneActorReferenceKey UMovieSceneActorReferenceTrackRecorder::ActorToReferenceKey(AActor* InActor) const
+{	
+	if (InActor != nullptr)
+	{
+		ULevelSequence* RootLevelSequence = OwningTakeRecorderSource->GetRootLevelSequence();
+		check(RootLevelSequence);
+		
+		FMovieSceneObjectBindingID Binding;
+		
+		const TOptional<FMovieSceneSpawnableAnnotation> Spawnable = FMovieSceneSpawnableAnnotation::Find(InActor);
+		if (Spawnable.IsSet())
+		{
+			// @TODO: This pathway is designed for spawnables when recording over an existing recording, but hasn't been fully implemented yet.
+			FText OutError;
+			if (const TSharedPtr<ISequencer> SequencerPtr = TakesUtils::OpenSequencer(RootLevelSequence, &OutError))
+			{
+				// Check whether the spawnable is underneath the current sequence, if so, we can remap it to a local sequence ID
+				Binding = UE::MovieScene::FRelativeObjectBindingID(SequencerPtr->GetFocusedTemplateID(), Spawnable->SequenceID, Spawnable->ObjectBindingID, *SequencerPtr);
+			}
+		}
+		else
+		{
+			FGuid ParentActorId = TakesUtils::ResolveActorFromSequence(InActor, RootLevelSequence);
+				
+			if (!ParentActorId.IsValid())
+			{
+				UE::Sequencer::FCreateBindingParams Params;
+				// Spawnable flag being accurate is critical for proper binding to actors depending on the actor & environment settings.
+				Params.bSpawnable = !OwningTakeRecorderSource->IsOtherActorPossessable(InActor);
+				if (Params.bSpawnable)
+				{
+					Params.bAllowCustomBinding = true;
+					Params.BindingNameOverride = InActor->GetActorLabel();
+				}
+				ParentActorId = FSequencerUtilities::CreateOrReplaceBinding(nullptr, RootLevelSequence, InActor, Params);
+				// If this is before a recording started or the actor is already being recorded, we don't need to create the source now
+				// because the system will handle it normally instead.
+				const bool bCreateNewSourceNow = bIsRecording && !OwningTakeRecorderSource->IsOtherActorBeingRecorded(InActor);
+				UTakeRecorderSources* Sources = RootLevelSequence->FindOrAddMetaData<UTakeRecorderSources>();
+				Sources->OnNewActorBindingAdded(InActor, ParentActorId, RootLevelSequence, bCreateNewSourceNow);
+			}
+			Binding = UE::MovieScene::FRelativeObjectBindingID(ParentActorId);
+		}
+
+		if (Binding.IsValid())
+		{
+			FMovieSceneActorReferenceKey NewKey;
+			NewKey.Object = MoveTemp(Binding);
+			return NewKey;
+		}
+	}
+	
+	return FMovieSceneActorReferenceKey();
+}
+
+AActor* UMovieSceneActorReferenceTrackRecorder::EvaluateBinding(const UObject* OwningObject) const
+{
+	if (OwningObject == nullptr)
+	{
+		OwningObject = ObjectToRecord.Get();
+	}
+	
+	check(OwningObject);
+	
+	const FProperty* PropertyToRecord = PropertyBinding->GetProperty(*OwningObject);
+	check(PropertyToRecord);
+		
+	AActor* ActorValue = nullptr;
+	if (PropertyToRecord->IsA<FSoftObjectProperty>())
+	{
+		const TSoftObjectPtr<AActor> ActorValueSoftObjectPtr = PropertyBinding->GetCurrentValue<TSoftObjectPtr<AActor>>(*OwningObject);
+		ActorValue = ActorValueSoftObjectPtr.LoadSynchronous();
+	}
+	else if (PropertyToRecord->IsA<FObjectProperty>())
+	{
+		ActorValue = Cast<AActor>(PropertyBinding->GetCurrentValue<TObjectPtr<UObject>>(*OwningObject));
+	}
+		
+	return ActorValue;
+}
diff --git a/Engine/Plugins/VirtualProduction/Takes/Source/TakeTrackRecorders/Private/TakeTrackRecordersModule.cpp b/Engine/Plugins/VirtualProduction/Takes/Source/TakeTrackRecorders/Private/TakeTrackRecordersModule.cpp
--- a/Engine/Plugins/VirtualProduction/Takes/Source/TakeTrackRecorders/Private/TakeTrackRecordersModule.cpp
+++ b/Engine/Plugins/VirtualProduction/Takes/Source/TakeTrackRecorders/Private/TakeTrackRecordersModule.cpp
@@ -6,6 +6,7 @@
 #include "Features/IModularFeatures.h"
 #include "TrackRecorders/MovieScene3DTransformTrackRecorder.h"
 #include "TrackRecorders/MovieScene3DAttachTrackRecorder.h"
+#include "TrackRecorders/MovieSceneActorReferenceTrackRecorder.h"
 #include "TrackRecorders/MovieSceneParticleTrackRecorder.h"
 #include "TrackRecorders/MovieSceneSpawnTrackRecorder.h"
 #include "TrackRecorders/MovieSceneVisibilityTrackRecorder.h"
@@ -25,6 +26,7 @@ class FTakeTrackRecordersModule : public IModuleInterface
 		IModularFeatures::Get().RegisterModularFeature(MovieSceneTrackRecorderFactoryName, &MovieSceneParticleTrackRecorder);
 		IModularFeatures::Get().RegisterModularFeature(MovieSceneTrackRecorderFactoryName, &MovieSceneSpawnTrackRecorder);
 		IModularFeatures::Get().RegisterModularFeature(MovieSceneTrackRecorderFactoryName, &MovieSceneVisibilityTrackRecorder);
+		IModularFeatures::Get().RegisterModularFeature(MovieSceneTrackRecorderFactoryName, &MovieSceneActorReferenceTrackRecorder);
 	}
 
 	virtual void ShutdownModule() override
@@ -36,6 +38,7 @@ class FTakeTrackRecordersModule : public IModuleInterface
 		IModularFeatures::Get().UnregisterModularFeature(MovieSceneTrackRecorderFactoryName, &MovieSceneParticleTrackRecorder);
 		IModularFeatures::Get().UnregisterModularFeature(MovieSceneTrackRecorderFactoryName, &MovieSceneSpawnTrackRecorder);
 		IModularFeatures::Get().UnregisterModularFeature(MovieSceneTrackRecorderFactoryName, &MovieSceneVisibilityTrackRecorder);
+		IModularFeatures::Get().UnregisterModularFeature(MovieSceneTrackRecorderFactoryName, &MovieSceneActorReferenceTrackRecorder);
 	}
 
 	FMovieScene3DTransformTrackRecorderFactory MovieScene3DTransformTrackRecorder;
@@ -44,6 +47,7 @@ class FTakeTrackRecordersModule : public IModuleInterface
 	FMovieSceneParticleTrackRecorderFactory MovieSceneParticleTrackRecorder;
 	FMovieSceneSpawnTrackRecorderFactory MovieSceneSpawnTrackRecorder;
 	FMovieSceneVisibilityTrackRecorderFactory MovieSceneVisibilityTrackRecorder;
+	FMovieSceneActorReferenceTrackRecorderFactory MovieSceneActorReferenceTrackRecorder;
 };
 
 IMPLEMENT_MODULE(FTakeTrackRecordersModule, TakeTrackRecorders)
diff --git a/Engine/Plugins/VirtualProduction/Takes/Source/TakeTrackRecorders/Public/TrackRecorders/IMovieSceneTrackRecorderFactory.h b/Engine/Plugins/VirtualProduction/Takes/Source/TakeTrackRecorders/Public/TrackRecorders/IMovieSceneTrackRecorderFactory.h
--- a/Engine/Plugins/VirtualProduction/Takes/Source/TakeTrackRecorders/Public/TrackRecorders/IMovieSceneTrackRecorderFactory.h
+++ b/Engine/Plugins/VirtualProduction/Takes/Source/TakeTrackRecorders/Public/TrackRecorders/IMovieSceneTrackRecorderFactory.h
@@ -79,7 +79,12 @@ public:
 	virtual bool IsSerializable() const { return false; }
 
 	/**
-	*  Unique Name that's stored in the serialized manifest when it's recoreded
+	*  Unique Name that's stored in the serialized manifest when it's recorded
 	*/
 	virtual FName GetSerializedType() const { return FName(); }
+	
+	/**
+	 * Allow a factory to handle individual properties on an object, regardless of the property name.
+	 */
+	virtual bool CanHandlePropertyInstance() const { return false; }
 };
diff --git a/Engine/Plugins/VirtualProduction/Takes/Source/TakeTrackRecorders/Public/TrackRecorders/IMovieSceneTrackRecorderHost.h b/Engine/Plugins/VirtualProduction/Takes/Source/TakeTrackRecorders/Public/TrackRecorders/IMovieSceneTrackRecorderHost.h
--- a/Engine/Plugins/VirtualProduction/Takes/Source/TakeTrackRecorders/Public/TrackRecorders/IMovieSceneTrackRecorderHost.h
+++ b/Engine/Plugins/VirtualProduction/Takes/Source/TakeTrackRecorders/Public/TrackRecorders/IMovieSceneTrackRecorderHost.h
@@ -191,4 +191,9 @@ public:
 
 	/**/
 	virtual FMovieSceneSequenceID GetSequenceID() const = 0;
+	
+	/** 
+	 * If the given actor is possessable.
+	 */
+	virtual bool IsOtherActorPossessable(AActor* InActor) const = 0;
 };
diff --git a/Engine/Plugins/VirtualProduction/Takes/Source/TakeTrackRecorders/Public/TrackRecorders/MovieSceneActorReferenceTrackRecorder.h b/Engine/Plugins/VirtualProduction/Takes/Source/TakeTrackRecorders/Public/TrackRecorders/MovieSceneActorReferenceTrackRecorder.h
new file mode 100644
--- /dev/null
+++ b/Engine/Plugins/VirtualProduction/Takes/Source/TakeTrackRecorders/Public/TrackRecorders/MovieSceneActorReferenceTrackRecorder.h
@@ -0,0 +1,77 @@
+// Copyright Epic Games, Inc. All Rights Reserved.
+
+#pragma once
+
+#include "IMovieSceneTrackRecorderFactory.h"
+#include "MovieSceneTrackRecorder.h"
+#include "Sections/MovieSceneActorReferenceSection.h"
+#include "Sections/MovieSceneBoolSection.h"
+#include "TrackInstancePropertyBindings.h"
+
+#include "MovieSceneActorReferenceTrackRecorder.generated.h"
+
+#define UE_API TAKETRACKRECORDERS_API
+
+
+class FMovieSceneActorReferenceTrackRecorderFactory : public IMovieSceneTrackRecorderFactory
+{
+public:
+	virtual ~FMovieSceneActorReferenceTrackRecorderFactory() = default;
+
+	UE_API virtual bool CanRecordObject(UObject* InObjectToRecord) const override;
+	UE_API virtual UMovieSceneTrackRecorder* CreateTrackRecorderForObject() const override;
+
+	UE_API virtual bool CanRecordProperty(UObject* InObjectToRecord, FProperty* InPropertyToRecord) const override;
+	UE_API virtual UMovieSceneTrackRecorder* CreateTrackRecorderForProperty(UObject* InObjectToRecord, const FName& InPropertyToRecord) const override;
+
+	virtual FText GetDisplayName() const override { return NSLOCTEXT("MovieSceneActorReferenceTrackRecorderFactory", "DisplayName", "Actor Reference Track"); }
+	
+	virtual bool CanHandlePropertyInstance() const override { return true; }
+};
+
+UCLASS(MinimalAPI, BlueprintType)
+class UMovieSceneActorReferenceTrackRecorder : public UMovieSceneTrackRecorder
+{
+	GENERATED_BODY()
+	
+public:
+	UE_API void CreateBinding(const FName& InPropertyToRecord);
+	
+protected:
+	// UMovieSceneTrackRecorder Interface
+	UE_API virtual void CreateTrackImpl() override;
+	UE_API virtual void RecordSampleImpl(const FQualifiedFrameTime& CurrentTime) override;
+	UE_API virtual void FinalizeTrackImpl() override;
+	virtual UMovieSceneSection* GetMovieSceneSection() const override { return Cast<UMovieSceneSection>(MovieSceneSection.Get()); }
+	// UMovieSceneTrackRecorder Interface
+
+	UE_API void RemoveRedundantTracks();
+
+private:
+	/** Convert an actor to a reference key, locating or creating a binding as necessary, taking into account if the actor is a spawnable or not. */
+	FMovieSceneActorReferenceKey ActorToReferenceKey(AActor* InActor) const;
+	
+	/** 
+	 * Evaluate the current binding returning the actor value.
+	 * @param OwningObject Optional object to use for binding evaluation, defaults to the owning take recorder source.
+	 */
+	AActor* EvaluateBinding(const UObject* OwningObject = nullptr) const;
+	
+private:
+	/** Section to record to */
+	TWeakObjectPtr<UMovieSceneActorReferenceSection> MovieSceneSection;
+	
+	/** Binding for this property */
+	TSharedPtr<FTrackInstancePropertyBindings> PropertyBinding;
+	
+	/** Value of the last evaluated binding. */
+	TWeakObjectPtr<AActor> LastBindingValue;
+	
+	/** Flag used to determine whether the first key needs to be set */
+	bool bSetFirstKey;
+	
+	/** True after a recording has started and until it stops. */
+	bool bIsRecording = false;
+};
+
+#undef UE_API
diff --git a/Engine/Plugins/VirtualProduction/Takes/Source/TakeTrackRecorders/TakeTrackRecorders.Build.cs b/Engine/Plugins/VirtualProduction/Takes/Source/TakeTrackRecorders/TakeTrackRecorders.Build.cs
--- a/Engine/Plugins/VirtualProduction/Takes/Source/TakeTrackRecorders/TakeTrackRecorders.Build.cs
+++ b/Engine/Plugins/VirtualProduction/Takes/Source/TakeTrackRecorders/TakeTrackRecorders.Build.cs
@@ -18,6 +18,7 @@ public class TakeTrackRecorders : ModuleRules
 				"SequenceRecorder", // For access to the individual UProperty type recorders
 				"SerializedRecorderInterface",
 				"UnrealEd",
+				"Sequencer"
             }
         );
 	}
diff --git a/Engine/Plugins/VirtualProduction/Takes/Source/TakesCore/Private/TakeRecorderSources.cpp b/Engine/Plugins/VirtualProduction/Takes/Source/TakesCore/Private/TakeRecorderSources.cpp
--- a/Engine/Plugins/VirtualProduction/Takes/Source/TakesCore/Private/TakeRecorderSources.cpp
+++ b/Engine/Plugins/VirtualProduction/Takes/Source/TakesCore/Private/TakeRecorderSources.cpp
@@ -111,6 +111,26 @@ UTakeRecorderSources::FOnSourceRemoved& UTakeRecorderSources::OnSourceRemoved()
 	return OnSourceRemovedEvent;
 }
 
+UTakeRecorderSources::FOnActorBindingAdded& UTakeRecorderSources::GetOnActorBindingAddedDelegate()
+{
+	static FOnActorBindingAdded OnActorBindingAddedDelegate;
+	return OnActorBindingAddedDelegate;
+}
+
+const TMap<TObjectPtr<AActor>, FGuid>& UTakeRecorderSources::GetNewActorBindings()
+{
+	return NewActorBindings;
+}
+
+void UTakeRecorderSources::OnNewActorBindingAdded(AActor* InActor, const FGuid& InBindingGuid, ULevelSequence* InLevelSequence, bool bCreateNewSourceNow)
+{
+	NewActorBindings.Add(InActor, InBindingGuid);
+	if (bCreateNewSourceNow)
+	{
+		GetOnActorBindingAddedDelegate().Broadcast(InActor, InLevelSequence);
+	}
+}
+
 void UTakeRecorderSources::PostEditChangeProperty(FPropertyChangedEvent& PropertyChangedEvent)
 {
 	Super::PostEditChangeProperty(PropertyChangedEvent);
@@ -133,8 +153,18 @@ void UTakeRecorderSources::PostEditChangeProperty(FPropertyChangedEvent& Propert
 	}
 }
 
+void UTakeRecorderSources::StartPreRecordingRecursive(const TArray<UTakeRecorderSource*>& InSources, ULevelSequence* InRootSequence,
+	TArray<UTakeRecorderSource*>& NewSourcesOut, FManifestSerializer* InManifestSerializer)
+{
+	TSet<FString> UsedTrackNames;
+	PreRecordingRecursive(InSources, InRootSequence, NewSourcesOut, InManifestSerializer, UsedTrackNames);
+	
+	// Make sure any recently added bindings are cleared out in case there wasn't a source to clear them.
+	NewActorBindings.Empty();
+}
+
 void UTakeRecorderSources::PreRecordingRecursive(TArray<UTakeRecorderSource*> InSources, ULevelSequence* InRootSequence,
-	TArray<UTakeRecorderSource*>& NewSourcesOut, FManifestSerializer* InManifestSerializer, TSet<FString>& UsedTrackNames)
+                                                 TArray<UTakeRecorderSource*>& NewSourcesOut, FManifestSerializer* InManifestSerializer, TSet<FString>& UsedTrackNames)
 {
 
 	TArray<UTakeRecorderSource*> NewSources;
@@ -313,8 +343,7 @@ void UTakeRecorderSources::StartRecordingPreRecordedSources(const FQualifiedFram
 void UTakeRecorderSources::PreRecordSources(TArray<UTakeRecorderSource *> InSources)
 {
 	PreRecordedSources.Reset(0);
-	TSet<FString> UsedTrackNames;
-	PreRecordingRecursive(InSources, CachedLevelSequence, PreRecordedSources, CachedManifestSerializer, UsedTrackNames);
+	StartPreRecordingRecursive(InSources, CachedLevelSequence, PreRecordedSources, CachedManifestSerializer);
 }
 
 
@@ -323,8 +352,7 @@ void UTakeRecorderSources::StartRecordingSource(TArray<UTakeRecorderSource *> In
 	// This calls PreRecording recursively on every source so that all sources that get added by another source
 	// have had PreRecording called.
 	TArray<UTakeRecorderSource *> NewSources;
-	TSet<FString> UsedTrackNames;
-	PreRecordingRecursive(InSources, CachedLevelSequence, NewSources, CachedManifestSerializer, UsedTrackNames);
+	StartPreRecordingRecursive(InSources, CachedLevelSequence, NewSources, CachedManifestSerializer);
 	if (NewSources.Num() > 0)
 	{
 		InSources.Append(NewSources);
@@ -596,6 +624,7 @@ void UTakeRecorderSources::StopRecording(class ULevelSequence* InSequence, const
 	SourceSubSequenceMap.Empty();
 	ActiveSubSections.Empty();
 	CreatedManifestSerializers.Empty();
+	NewActorBindings.Empty();
 	CachedManifestSerializer = nullptr;
 	CachedLevelSequence = nullptr;
 
diff --git a/Engine/Plugins/VirtualProduction/Takes/Source/TakesCore/Private/TakesUtils.cpp b/Engine/Plugins/VirtualProduction/Takes/Source/TakesCore/Private/TakesUtils.cpp
--- a/Engine/Plugins/VirtualProduction/Takes/Source/TakesCore/Private/TakesUtils.cpp
+++ b/Engine/Plugins/VirtualProduction/Takes/Source/TakesCore/Private/TakesUtils.cpp
@@ -326,6 +326,26 @@ FQualifiedFrameTime GetRecordTime(TSharedPtr<ISequencer> Sequencer, ULevelSequen
 
 	return RecordTime;
 }
+
+FGuid ResolveActorFromSequence(const AActor* InActor, const ULevelSequence* CurrentSequence)
+{
+	UMovieScene* MovieScene = CurrentSequence->GetMovieScene();
+
+	// Look through all Possessables in the sequence to see if there's one with the same name as the actor. We purposely do not look at 
+	// Spawnables so that recording as spawnable will always create a new spawnable.
+	for (int32 PossessableCount = 0; PossessableCount < MovieScene->GetPossessableCount(); ++PossessableCount)
+	{
+		const FMovieScenePossessable& Possessable = MovieScene->GetPossessable(PossessableCount);
+		if (Possessable.GetName() == InActor->GetActorLabel())
+		{
+			return Possessable.GetGuid();
+		}
+	}
+
+	// There's no Possessable with the same name as the actor, so this actor hasn't been added
+	// to the sequence yet.
+	return FGuid();
+}
 }
 
 #undef LOCTEXT_NAMESPACE
diff --git a/Engine/Plugins/VirtualProduction/Takes/Source/TakesCore/Public/TakeRecorderSources.h b/Engine/Plugins/VirtualProduction/Takes/Source/TakesCore/Public/TakeRecorderSources.h
--- a/Engine/Plugins/VirtualProduction/Takes/Source/TakesCore/Public/TakeRecorderSources.h
+++ b/Engine/Plugins/VirtualProduction/Takes/Source/TakesCore/Public/TakeRecorderSources.h
@@ -31,6 +31,7 @@ struct FTakeRecorderSourcesSettings
 };
 
 class UWorld;
+class ULevelSequence;
 class FViewport;
 class FCanvas;
 
@@ -136,6 +137,28 @@ public:
 	/** Multicast delegate for when any source is removed. */
 	static UE_API FOnSourceRemoved& OnSourceRemoved();
 	
+	DECLARE_MULTICAST_DELEGATE_TwoParams(FOnActorBindingAdded, AActor*, ULevelSequence*);
+	/** Multicast delegate for when an actor binding is added. */
+	static UE_API FOnActorBindingAdded& GetOnActorBindingAddedDelegate();
+
+	/** Actors that have had bindings created recently. */
+	UE_API const TMap<TObjectPtr<AActor>, FGuid>& GetNewActorBindings();
+	
+	/**
+	 * Notifies sources that a new actor binding has been added.
+	 * 
+	 * @param InActor The actor the binding was added for.
+	 * @param InBindingGuid The binding guid created.
+	 * @param InLevelSequence The level sequence of the binding.
+	 * @param bCreateNewSourceNow If a new source should be created for the actor this moment. This should be true
+	 * if a recording is already in progress.
+	 */
+	UE_API void OnNewActorBindingAdded(AActor* InActor, const FGuid& InBindingGuid, ULevelSequence* InLevelSequence, bool bCreateNewSourceNow);
+
+private:
+	UPROPERTY(Transient)
+	TMap<TObjectPtr<AActor>, FGuid> NewActorBindings;
+	
 public:
 
 	/*
@@ -207,6 +230,10 @@ private:
 	DECLARE_MULTICAST_DELEGATE(FOnSourcesChanged);
 	FOnSourcesChanged OnSourcesChangedEvent;
 
+	/** Calls PreRecording on sources recursively allowing them to create other sources which properly get PreRecording called on them as well. */
+	void StartPreRecordingRecursive(const TArray<UTakeRecorderSource*>& InSources, ULevelSequence* InRootSequence, TArray<UTakeRecorderSource*>& NewSourcesOut,
+		FManifestSerializer* InManifestSerializer);
+	
 	/** Calls PreRecording on sources recursively allowing them to create other sources which properly get PreRecording called on them as well. */
 	UE_API void PreRecordingRecursive(TArray<UTakeRecorderSource*> InSources, ULevelSequence* InRootSequence, TArray<UTakeRecorderSource*>& NewSourcesOut,
 		FManifestSerializer* InManifestSerializer, TSet<FString>& UsedTrackNames);
diff --git a/Engine/Plugins/VirtualProduction/Takes/Source/TakesCore/Public/TakesUtils.h b/Engine/Plugins/VirtualProduction/Takes/Source/TakesCore/Public/TakesUtils.h
--- a/Engine/Plugins/VirtualProduction/Takes/Source/TakesCore/Public/TakesUtils.h
+++ b/Engine/Plugins/VirtualProduction/Takes/Source/TakesCore/Public/TakesUtils.h
@@ -175,4 +175,7 @@ namespace TakesUtils
 	/* The time at which to record. Taken from the Sequencer global time, otherwise based on timecode */
 	TAKESCORE_API FQualifiedFrameTime GetRecordTime(TSharedPtr<ISequencer> Sequencer, ULevelSequence* SequenceAsset,
 		const FTimecode& TimecodeAtStart, bool bStartAtCurrentTimecode, float TimeDilation = 1.f);
+	
+	/** Retrieve an actor guid for moviescene. */
+	TAKESCORE_API FGuid ResolveActorFromSequence(const AActor* InActor, const ULevelSequence* CurrentSequence);
 }
diff --git a/Engine/Source/Editor/Sequencer/Private/SequencerUtilities.cpp b/Engine/Source/Editor/Sequencer/Private/SequencerUtilities.cpp
--- a/Engine/Source/Editor/Sequencer/Private/SequencerUtilities.cpp
+++ b/Engine/Source/Editor/Sequencer/Private/SequencerUtilities.cpp
@@ -3132,7 +3132,9 @@ FGuid CreateGenericBinding(TSharedPtr<ISequencer> Sequencer, UMovieSceneSequence
 		}
 
 		// If we're not trying to replace a binding, and the object is a spawnable, try and bind to that first
-		if (InParams.ReplacementGuid.IsValid() || (Sequencer && (!CurrentObject || !NewPossessable->BindSpawnableObject(Sequencer->GetFocusedTemplateID(), CurrentObject, Sequencer->GetSharedPlaybackState()))))
+		if (InParams.ReplacementGuid.IsValid()
+			|| (Sequencer && (!CurrentObject || !NewPossessable->BindSpawnableObject(Sequencer->GetFocusedTemplateID(), CurrentObject, Sequencer->GetSharedPlaybackState())))
+			|| (!Sequencer && !InParams.bSpawnable)) /** If binding without sequencer we still need to call AddBinding */
 		{
 			FUniversalObjectLocator Locator;
 			if (CurrentObject && (!OwnerSequence->MakeLocatorForObject(CurrentObject, Context, Locator) || Locator.IsEmpty()))
