TakeRecorder: Change new IMovieSceneTrackRecorderHost interface function from pure virtual to virtual to fix build issue 1219633

#Synced-CL 50653618

diff --git a/Engine/Plugins/VirtualProduction/Takes/Source/TakeTrackRecorders/Public/TrackRecorders/IMovieSceneTrackRecorderHost.h b/Engine/Plugins/VirtualProduction/Takes/Source/TakeTrackRecorders/Public/TrackRecorders/IMovieSceneTrackRecorderHost.h
--- a/Engine/Plugins/VirtualProduction/Takes/Source/TakeTrackRecorders/Public/TrackRecorders/IMovieSceneTrackRecorderHost.h
+++ b/Engine/Plugins/VirtualProduction/Takes/Source/TakeTrackRecorders/Public/TrackRecorders/IMovieSceneTrackRecorderHost.h
@@ -195,5 +195,5 @@ public:
 	/** 
 	 * If the given actor is possessable.
 	 */
-	virtual bool IsOtherActorPossessable(AActor* InActor) const = 0;
+	virtual bool IsOtherActorPossessable(AActor* InActor) const { return false; }
 };
