TakeRecorder: Add search bar to the settings ara of the panel.

#rnx

#Synced-CL 50633722

diff --git a/Engine/Plugins/VirtualProduction/Takes/Source/TakeRecorder/Private/Widgets/SLevelSequenceTakeEditor.cpp b/Engine/Plugins/VirtualProduction/Takes/Source/TakeRecorder/Private/Widgets/SLevelSequenceTakeEditor.cpp
--- a/Engine/Plugins/VirtualProduction/Takes/Source/TakeRecorder/Private/Widgets/SLevelSequenceTakeEditor.cpp
+++ b/Engine/Plugins/VirtualProduction/Takes/Source/TakeRecorder/Private/Widgets/SLevelSequenceTakeEditor.cpp
@@ -40,6 +40,7 @@
 // Slate includes
 #include "Widgets/Input/SComboBox.h"
 #include "Widgets/Input/SEditableTextBox.h"
+#include "Widgets/Input/SSearchBox.h"
 #include "Widgets/Layout/SSplitter.h"
 #include "Widgets/Layout/SScrollBox.h"
 #include "Widgets/Layout/SBorder.h"
@@ -135,7 +136,23 @@ void SLevelSequenceTakeEditor::Construct(const FArguments& InArgs)
 		+ SSplitter::Slot()
 		.Value(.5f)
 		[
-			DetailsBox.ToSharedRef()
+			SNew(SVerticalBox)
+
+			+ SVerticalBox::Slot()
+			.AutoHeight()
+			.Padding(4.f, 4.f, 4.f, 2.f)
+			[
+				SAssignNew(SearchBox, SSearchBox)
+				.HintText(LOCTEXT("SearchSettingsHint", "Search Settings"))
+				.OnTextChanged(this, &SLevelSequenceTakeEditor::OnSearchTextChanged)
+				.OnTextCommitted(this, &SLevelSequenceTakeEditor::OnSearchTextCommitted)
+			]
+
+			+ SVerticalBox::Slot()
+			.FillHeight(1.f)
+			[
+				DetailsBox.ToSharedRef()
+			]
 		]
 	];
 }
@@ -281,6 +298,35 @@ void SLevelSequenceTakeEditor::OnSourcesSelectionChanged(TSharedPtr<ITakeRecorde
 	bRequestDetailsRefresh = true;
 }
 
+void SLevelSequenceTakeEditor::OnSearchTextChanged(const FText& InFilterText)
+{
+	CurrentSearchText = InFilterText;
+	const bool bIsSearchActive = !InFilterText.IsEmpty();
+
+	for (const auto& Pair : ClassToDetailsView)
+	{
+		if (Pair.Value.IsValid())
+		{
+			Pair.Value->SetSearchText(InFilterText);
+
+			// Collapse details views that have no matching properties to avoid empty "All results have been filtered" messages
+			if (bIsSearchActive && Pair.Value->GetPropertiesInOrderDisplayed().Num() == 0)
+			{
+				Pair.Value->SetVisibility(EVisibility::Collapsed);
+			}
+			else
+			{
+				Pair.Value->SetVisibility(EVisibility::Visible);
+			}
+		}
+	}
+}
+
+void SLevelSequenceTakeEditor::OnSearchTextCommitted(const FText& InFilterText, ETextCommit::Type CommitType)
+{
+	OnSearchTextChanged(InFilterText);
+}
+
 void SLevelSequenceTakeEditor::OnDetailsPropertiesChanged(const FPropertyChangedEvent& InEvent)
 {
 	OnDetailsPropertiesChangedEvent.ExecuteIfBound(InEvent);
@@ -306,7 +352,9 @@ void SLevelSequenceTakeEditor::AddDetails(const TPair<const UClass*, TArray<UObj
 {
 	FPropertyEditorModule& PropertyEditorModule = FModuleManager::Get().LoadModuleChecked<FPropertyEditorModule>("PropertyEditor");
 	FDetailsViewArgs DetailsViewArgs;
-	DetailsViewArgs.bAllowSearch = false;
+	DetailsViewArgs.bAllowSearch = true;
+	DetailsViewArgs.bCustomFilterAreaLocation = true;
+	DetailsViewArgs.bShowOptions = false;
 	DetailsViewArgs.NameAreaSettings = FDetailsViewArgs::HideNameArea;
 	DetailsViewArgs.bHideSelectionTip = true;
 	DetailsViewArgs.bShowScrollBar = false;
@@ -346,7 +394,17 @@ void SLevelSequenceTakeEditor::AddDetails(const TPair<const UClass*, TArray<UObj
 
 		Details->OnFinishedChangingProperties().AddSP(this, &SLevelSequenceTakeEditor::OnDetailsPropertiesChanged);
 		OnDetailsViewAddedEvent.ExecuteIfBound(Details);
-		
+
+		if (!CurrentSearchText.IsEmpty())
+		{
+			Details->SetSearchText(CurrentSearchText);
+
+			if (Details->GetPropertiesInOrderDisplayed().Num() == 0)
+			{
+				Details->SetVisibility(EVisibility::Collapsed);
+			}
+		}
+
 		DetailsBox->AddSlot()
 			[
 				Details
diff --git a/Engine/Plugins/VirtualProduction/Takes/Source/TakeRecorder/Private/Widgets/SLevelSequenceTakeEditor.h b/Engine/Plugins/VirtualProduction/Takes/Source/TakeRecorder/Private/Widgets/SLevelSequenceTakeEditor.h
--- a/Engine/Plugins/VirtualProduction/Takes/Source/TakeRecorder/Private/Widgets/SLevelSequenceTakeEditor.h
+++ b/Engine/Plugins/VirtualProduction/Takes/Source/TakeRecorder/Private/Widgets/SLevelSequenceTakeEditor.h
@@ -12,6 +12,7 @@ struct FAssetData;
 struct ITakeRecorderSourceTreeItem;
 
 class SScrollBox;
+class SSearchBox;
 class ULevelSequence;
 class IDetailsView;
 class UTakeRecorderSource;
@@ -88,6 +89,12 @@ private:
 	bool CanAddSourceFromClass(TSubclassOf<UTakeRecorderSource> SourceClass);
 	
 	void OnSourcesSelectionChanged(TSharedPtr<ITakeRecorderSourceTreeItem>, ESelectInfo::Type);
+	/**
+	 * When the search text changes, propagate to all details views.
+	 */
+	void OnSearchTextChanged(const FText& InFilterText);
+	void OnSearchTextCommitted(const FText& InFilterText, ETextCommit::Type CommitType);
+
 	/**
 	 * When a details box has a property changed.
 	 */
@@ -103,7 +110,9 @@ private:
 	TWeakObjectPtr<ULevelSequence> CachedLevelSequence;
 
 	TSharedPtr<STakeRecorderSources> SourcesWidget;
+	TSharedPtr<SSearchBox> SearchBox;
 	TSharedPtr<SScrollBox> DetailsBox;
+	FText CurrentSearchText;
 	TMap<FObjectKey, TSharedPtr<IDetailsView>> ClassToDetailsView;
 
 	TArray<TWeakObjectPtr<>> ExternalSettingsObjects;
