TakeRecorder: Improve actor reference recording to support subsequences.
Rework implementation, remove/rename/make private some previously added (unreleased) API.

#Synced-CL 50791565

diff --git a/Engine/Plugins/VirtualProduction/Takes/Source/TakeRecorderSources/Private/TakeRecorderActorSource.cpp b/Engine/Plugins/VirtualProduction/Takes/Source/TakeRecorderSources/Private/TakeRecorderActorSource.cpp
--- a/Engine/Plugins/VirtualProduction/Takes/Source/TakeRecorderSources/Private/TakeRecorderActorSource.cpp
+++ b/Engine/Plugins/VirtualProduction/Takes/Source/TakeRecorderSources/Private/TakeRecorderActorSource.cpp
@@ -249,25 +249,7 @@ TArray<UTakeRecorderSource*> UTakeRecorderActorSource::PreRecording(ULevelSequen
 	// Look to see if the movie scene already has our object binding in it (which is common if we're recording a new take) so we can
 	// replace the data that was already there.
 	CachedObjectBindingGuid = ResolveActorFromSequence(ActorToRecord, TargetLevelSequence);
-	
-	// Check if the binding was just created by an actor reference track that was processed by another actor source, prior to this one.
-	// In this case the existing binding is the one we want to use.
-	UTakeRecorderSources* SourcesList = GetTypedOuter<UTakeRecorderSources>();
-	check(SourcesList);
-	const FGuid* ExistingBindingGuid = SourcesList->GetNewActorBindings().Find(ActorToRecord);
-	const bool bIsExistingBindingValid = ExistingBindingGuid != nullptr;
-	if (bIsExistingBindingValid)
-	{
-		if (!CachedObjectBindingGuid.IsValid() )
-		{
-			CachedObjectBindingGuid = *ExistingBindingGuid;
-		}
-		ensure(CachedObjectBindingGuid == *ExistingBindingGuid);
-	}
-	else
-	{
-		CleanExistingDataFromSequence(CachedObjectBindingGuid, *TargetLevelSequence);
-	}
+	CleanExistingDataFromSequence(CachedObjectBindingGuid, *TargetLevelSequence);
 		
 	// Initialize the header for this actor in the Manifest Serializer for streaming data capture.
 	FName SerializedType("Actor");
@@ -284,12 +266,9 @@ TArray<UTakeRecorderSource*> UTakeRecorderActorSource::PreRecording(ULevelSequen
 		}
 		else
 		{
-			if (!bIsExistingBindingValid)
-			{
-				// Create a Possessable object binding in the Sequence and then bind it to our actor
-				CachedObjectBindingGuid = MovieScene->AddPossessable(ActorToRecord->GetActorLabel(), ActorToRecord->GetClass());
-				TargetLevelSequence->BindPossessableObject(CachedObjectBindingGuid, *ActorToRecord, ActorToRecord->GetWorld());
-			}
+			// Create a Possessable object binding in the Sequence and then bind it to our actor
+			CachedObjectBindingGuid = MovieScene->AddPossessable(ActorToRecord->GetActorLabel(), ActorToRecord->GetClass());
+			TargetLevelSequence->BindPossessableObject(CachedObjectBindingGuid, *ActorToRecord, ActorToRecord->GetWorld());
 			Header.bRecordToPossessable = true;
 		}
 	}
@@ -302,7 +281,7 @@ TArray<UTakeRecorderSource*> UTakeRecorderActorSource::PreRecording(ULevelSequen
 		CreateBindingParams.bSpawnable = true;
 		CreateBindingParams.bAllowCustomBinding = true;
 		CreateBindingParams.BindingNameOverride = ActorToRecord->GetActorLabel();
-		CachedObjectBindingGuid = ExistingBindingGuid ? *ExistingBindingGuid : FSequencerUtilities::CreateOrReplaceBinding(nullptr, InSequence, ActorToRecord, CreateBindingParams);
+		CachedObjectBindingGuid = FSequencerUtilities::CreateOrReplaceBinding(nullptr, InSequence, ActorToRecord, CreateBindingParams);
 		CachedObjectTemplate = Cast<AActor>(MovieSceneHelpers::GetObjectTemplate(InSequence, CachedObjectBindingGuid, MovieSceneHelpers::CreateTransientSharedPlaybackState(ActorToRecord, InSequence)));
 		
 		if (CachedObjectTemplate.IsValid())
@@ -366,6 +345,11 @@ TArray<UTakeRecorderSource*> UTakeRecorderActorSource::PreRecording(ULevelSequen
 	return AddedActorSources;
 }
 
+void UTakeRecorderActorSource::OnAllSourcesPreRecordingFinished()
+{
+	DoCreatePendingTrackRecorders();
+}
+
 void UTakeRecorderActorSource::EnsureParentHierarchyIsReferenced()
 {
 	if (!bRecordParentHierarchy)
@@ -484,6 +468,21 @@ void UTakeRecorderActorSource::CreateSectionRecordersRecursive(UObject* ObjectTo
 		ActorSerializer.WriteFrameData(ActorSerializer.FramesWritten, ActorCompFrame);
 	}
 
+	auto AddSectionRecorder = [this, ObjectToRecord, Guid](UMovieSceneTrackRecorder* SectionRecorder,
+		const TSubclassOf<UMovieSceneTrackRecorderSettings>& SettingsClass)
+		{
+			TrackRecorders.Add(SectionRecorder);
+			
+			FCreateTrackArgs CreateTrackArgs
+			{
+				SectionRecorder,
+				ObjectToRecord,
+				GetSettingsObjectForFactory(SettingsClass),
+				Guid
+			};
+			PendingTrackRecorders.Add(MoveTemp(CreateTrackArgs));
+		};
+	
 	// We need to iterate through the Property Map to see if the user wants to record this property or not
 	// We store the name of the Factory that can do the recording in the Property Map so for now we shortcut
 	// and just look up Factories by name instead of replicating all of the logic that goes into initializing
@@ -504,10 +503,8 @@ void UTakeRecorderActorSource::CreateSectionRecordersRecursive(UObject* ObjectTo
 			// cbb: Can't FText == FText?
 			if (Factory->GetDisplayName().ToString() == Property.RecorderName.ToString())
 			{
-				UMovieSceneTrackRecorder* SectionRecorder = Factory->CreateTrackRecorderForObject();
-				if (SectionRecorder)
+				if (UMovieSceneTrackRecorder* SectionRecorder = Factory->CreateTrackRecorderForObject())
 				{
-					TrackRecorders.Add(SectionRecorder);
 					if (Factory->IsSerializable())
 					{
 						AActor* ActorToRecord = Target.Get();
@@ -516,7 +513,8 @@ void UTakeRecorderActorSource::CreateSectionRecordersRecursive(UObject* ObjectTo
 						ActorSerializer.WriteFrameData(ActorSerializer.FramesWritten, ActorFrame);
 						SectionRecorder->SetSavedRecordingDirectory(ActorSerializer.GetLocalCaptureDir());
 					}
-					SectionRecorder->CreateTrack(this, ObjectToRecord, TargetLevelSequence->GetMovieScene(), GetSettingsObjectForFactory(Factory->GetSettingsClass()), Guid);
+					
+					AddSectionRecorder(SectionRecorder, Factory->GetSettingsClass());
 
 					bFoundRecorder = true;
 					break;
@@ -557,7 +555,6 @@ void UTakeRecorderActorSource::CreateSectionRecordersRecursive(UObject* ObjectTo
 				{
 					if (UMovieSceneTrackRecorder* SectionRecorder = TrackRecorderFactory.CreateTrackRecorderForProperty(ObjectToRecord, Property.PropertyName))
 					{
-						TrackRecorders.Add(SectionRecorder);
 						if (TrackRecorderFactory.IsSerializable())
 						{
 							AActor* ActorToRecord = Target.Get();
@@ -568,7 +565,8 @@ void UTakeRecorderActorSource::CreateSectionRecordersRecursive(UObject* ObjectTo
 							ActorSerializer.WriteFrameData(ActorSerializer.FramesWritten, ActorFrame);
 							SectionRecorder->SetSavedRecordingDirectory(ActorSerializer.GetLocalCaptureDir());
 						}
-						SectionRecorder->CreateTrack(this, ObjectToRecord, TargetLevelSequence->GetMovieScene(), GetSettingsObjectForFactory(TrackRecorderFactory.GetSettingsClass()), Guid);
+						
+						AddSectionRecorder(SectionRecorder, TrackRecorderFactory.GetSettingsClass());
 					}
 				};
 				
@@ -676,6 +674,9 @@ void UTakeRecorderActorSource::TickRecording(const FQualifiedFrameTime& CurrentS
 			// Create the Section Recorders required
 			CreateSectionRecordersRecursive(AddedComponent.Get(), ComponentPropertyMap, TraversedObjects);
 
+			// Ensure our track recorders are fully created after creating the initial sections.
+			DoCreatePendingTrackRecorders();
+			
 			// Update our numbers on the display
 			UpdateCachedNumberOfRecordedProperties();
 
@@ -811,6 +812,21 @@ UMovieSceneAnimationTrackRecorderEditorSettings* UTakeRecorderActorSource::GetBe
 	return GetMutableDefault<UMovieSceneAnimationTrackRecorderEditorSettings>();
 }
 
+void UTakeRecorderActorSource::DoCreatePendingTrackRecorders()
+{
+	for (const FCreateTrackArgs& PendingTrack : PendingTrackRecorders)
+	{
+		check(PendingTrack.TrackRecorder.IsValid());
+		check(PendingTrack.ObjectToRecord.IsValid());
+		check(PendingTrack.ObjectGuid.IsValid());
+		
+		PendingTrack.TrackRecorder->CreateTrack(this, PendingTrack.ObjectToRecord.Get(),
+			TargetLevelSequence->GetMovieScene(), PendingTrack.FactorySettings.Get(), PendingTrack.ObjectGuid);
+	}
+	
+	PendingTrackRecorders.Empty();
+}
+
 TArray<UTakeRecorderSource*> UTakeRecorderActorSource::PostRecording(ULevelSequence* InSequence, class ULevelSequence* InRootSequence, const bool bCancelled)
 {
 	FTakeRecorderParameters Parameters;
@@ -1716,10 +1732,6 @@ void UTakeRecorderActorSource::CreateNewActorSourceForReferencedActors()
 		}
 	}
 	
-	TArray<TObjectPtr<AActor>> NewActorBindingsArray;
-	SourcesList->GetNewActorBindings().GenerateKeyArray(NewActorBindingsArray);
-	NewReferencedActors.Append(MoveTemp(NewActorBindingsArray));
-
 	for (AActor* Actor : NewReferencedActors)
 	{
 		// Don't create a recording for this actor if there is an existing source for this actor. Another source may have added it
@@ -1776,6 +1788,11 @@ FGuid UTakeRecorderActorSource::GetRecordedActorGuid(class AActor* OtherActor) c
 	return TakeRecorderSourcesUtils::GetRecordedActorGuid(this, OtherActor);
 }
 
+ULevelSequence* UTakeRecorderActorSource::GetLevelSequence(const AActor* OtherActor)
+{
+	return TakeRecorderSourcesUtils::GetLevelSequence(this, OtherActor, RootLevelSequence);
+}
+
 FTransform UTakeRecorderActorSource::GetRecordedActorAnimationInitialRootTransform(class AActor* OtherActor) const
 {
 	return TakeRecorderSourcesUtils::GetRecordedActorAnimationInitialRootTransform(this, OtherActor);
diff --git a/Engine/Plugins/VirtualProduction/Takes/Source/TakeRecorderSources/Private/TakeRecorderSourcesModule.cpp b/Engine/Plugins/VirtualProduction/Takes/Source/TakeRecorderSources/Private/TakeRecorderSourcesModule.cpp
--- a/Engine/Plugins/VirtualProduction/Takes/Source/TakeRecorderSources/Private/TakeRecorderSourcesModule.cpp
+++ b/Engine/Plugins/VirtualProduction/Takes/Source/TakeRecorderSources/Private/TakeRecorderSourcesModule.cpp
@@ -166,7 +166,7 @@ void FTakeRecorderSourcesModule::StartupModule()
 	ISequencerModule& SequencerModule = FModuleManager::Get().LoadModuleChecked<ISequencerModule>("Sequencer");
 	OnSequencerCreatedHandle = SequencerModule.RegisterOnSequencerCreated(FOnSequencerCreated::FDelegate::CreateRaw(this, &FTakeRecorderSourcesModule::OnSequencerCreated));
 	
-	UTakeRecorderSources::GetOnActorBindingAddedDelegate().AddRaw(this, &FTakeRecorderSourcesModule::OnNewActorReferencingBindingAdded);
+	UTakeRecorderSources::OnAddSourceForActorBinding().BindRaw(this, &FTakeRecorderSourcesModule::OnNewActorReferenceBindingAdded);
 }
 
 void FTakeRecorderSourcesModule::ShutdownModule()
@@ -190,7 +190,7 @@ void FTakeRecorderSourcesModule::ShutdownModule()
 
 	UnregisterMenuExtensions();
 	
-	UTakeRecorderSources::GetOnActorBindingAddedDelegate().RemoveAll(this);
+	UTakeRecorderSources::OnAddSourceForActorBinding().Unbind();
 }
 
 void FTakeRecorderSourcesModule::RegisterCanRecordDelegate(FName HandleId, FCanRecordDelegate InDelegate)
@@ -777,12 +777,15 @@ void FTakeRecorderSourcesModule::OnSequencerCreated(TSharedRef<ISequencer> Seque
 	});
 }
 
-void FTakeRecorderSourcesModule::OnNewActorReferencingBindingAdded(AActor* InActor, ULevelSequence* InLevelSequence)
+UTakeRecorderSource* FTakeRecorderSourcesModule::OnNewActorReferenceBindingAdded(AActor* InActor, ULevelSequence* InLevelSequence)
 {	
 	if (UTakeRecorderActorSource* ActorSource = TakeRecorderSourcesUtils::StartRecordingActorDuringSession(InActor, InLevelSequence))
 	{
 		ActorSource->SetIsTemporarySource(true);
+		return ActorSource;
 	}
+	
+	return nullptr;
 }
 
 IMPLEMENT_MODULE(FTakeRecorderSourcesModule, TakeRecorderSources);
diff --git a/Engine/Plugins/VirtualProduction/Takes/Source/TakeRecorderSources/Private/TakeRecorderSourcesModule.h b/Engine/Plugins/VirtualProduction/Takes/Source/TakeRecorderSources/Private/TakeRecorderSourcesModule.h
--- a/Engine/Plugins/VirtualProduction/Takes/Source/TakeRecorderSources/Private/TakeRecorderSourcesModule.h
+++ b/Engine/Plugins/VirtualProduction/Takes/Source/TakeRecorderSources/Private/TakeRecorderSourcesModule.h
@@ -13,6 +13,7 @@ template<typename OptionalType>struct TOptional;
 
 class ULevelSequence;
 class UTakeRecorderSources;
+class UTakeRecorderSource;
 
 namespace UE::TakeRecorderSources
 {
@@ -69,6 +70,6 @@ private:
 
 	void OnSequencerCreated(TSharedRef<ISequencer> Sequencer);
 	
-	void OnNewActorReferencingBindingAdded(AActor* InActor, ULevelSequence* InLevelSequence);
+	UTakeRecorderSource* OnNewActorReferenceBindingAdded(AActor* InActor, ULevelSequence* InLevelSequence);
 };
 }
diff --git a/Engine/Plugins/VirtualProduction/Takes/Source/TakeRecorderSources/Private/TakeRecorderSourcesUtils.h b/Engine/Plugins/VirtualProduction/Takes/Source/TakeRecorderSources/Private/TakeRecorderSourcesUtils.h
--- a/Engine/Plugins/VirtualProduction/Takes/Source/TakeRecorderSources/Private/TakeRecorderSourcesUtils.h
+++ b/Engine/Plugins/VirtualProduction/Takes/Source/TakeRecorderSources/Private/TakeRecorderSourcesUtils.h
@@ -186,6 +186,36 @@ namespace TakeRecorderSourcesUtils
 		return FTransform::Identity;
 	}
 
+	/**
+	* Get the sequence of the level sequence the other actor is coming from.
+	* Used for setting cross sequence bindings.
+	* 
+	* @param Actor Actor to look for.
+	* @return The Sequence
+	*/
+	static ULevelSequence* GetLevelSequence(const UTakeRecorderSource* InSource, const AActor* Actor, ULevelSequence* RootLevelSequence)
+	{
+		check(InSource);
+		const UTakeRecorderSources* OwningSources = CastChecked<UTakeRecorderSources>(InSource->GetOuter());
+
+		for (UTakeRecorderSource* Source : OwningSources->GetSources())
+		{
+			if (UTakeRecorderActorSource* ActorSource = Cast<UTakeRecorderActorSource>(Source))
+			{
+				AActor* OtherTarget = ActorSource->Target.Get();
+				if (OtherTarget && OtherTarget->GetName() == Actor->GetName()) //at the end the target's may have changed.
+				{
+					if (ActorSource->TargetLevelSequence != RootLevelSequence)
+					{
+						return ActorSource->TargetLevelSequence;
+					}
+				}
+			}
+		}
+
+		return RootLevelSequence;
+	}
+	
 	/**
 	* Get the sequence id of the level sequence the other actor is coming from.
 	* Used for setting cross sequence bindings.
@@ -233,7 +263,7 @@ namespace TakeRecorderSourcesUtils
 		FPropertyChangedEvent PropertyChangedEvent(UTakeRecorderActorSource::StaticClass()->FindPropertyByName(GET_MEMBER_NAME_CHECKED(UTakeRecorderActorSource, Target)), EPropertyChangeType::ValueSet);
 		ActorSource->PostEditChangeProperty(PropertyChangedEvent);
 
-		// This has to be called after setting the Target and propegating the change event so that it has a chance to know what to record
+		// This has to be called after setting the Target and propagating the change event so that it has a chance to know what to record
 		// about the actor.
 		Sources->StartRecordingSource(TArray<UTakeRecorderSource*>({ ActorSource }), Sources->GetCachedFrameTime());
 		return ActorSource;
diff --git a/Engine/Plugins/VirtualProduction/Takes/Source/TakeRecorderSources/Public/TakeRecorderActorSource.h b/Engine/Plugins/VirtualProduction/Takes/Source/TakeRecorderSources/Public/TakeRecorderActorSource.h
--- a/Engine/Plugins/VirtualProduction/Takes/Source/TakeRecorderSources/Public/TakeRecorderActorSource.h
+++ b/Engine/Plugins/VirtualProduction/Takes/Source/TakeRecorderSources/Public/TakeRecorderActorSource.h
@@ -144,6 +144,7 @@ public:
 	// UTakeRecorderSource Interface
 	UE_API virtual bool IsValid() const override;
 	UE_API virtual TArray<UTakeRecorderSource*> PreRecording(ULevelSequence* InSequence, FMovieSceneSequenceID InSequenceID, ULevelSequence* InRootSequence, FManifestSerializer* InManifestSerializer) override;
+	UE_API virtual void OnAllSourcesPreRecordingFinished() override;
 	UE_API virtual void StartRecording(const FTimecode& InSectionStartTimecode, const FFrameNumber& InSectionFirstFrame, class ULevelSequence* InSequence) override;
 	UE_API virtual void TickRecording(const FQualifiedFrameTime& CurrentSequenceTime) override;
 	UE_API virtual void StopRecording(class ULevelSequence* InSequence) override;
@@ -200,7 +201,8 @@ protected:
 
 	/**
 	* This is called when recording starts to generate the Section Recorders for the actor and all components that it currently has,
-	* as well as again during runtime for any newly added components.
+	* as well as again during runtime for any newly added components. This instantiates the TrackRecorders needed, but CreateTrack will
+	* not be called until DoCreatePendingTrackRecorders() is called.
 	*/
 	void CreateSectionRecordersRecursive(UObject* ObjectToRecord, UActorRecorderPropertyMap* PropertyMap, TArray<UObject*>& TraversedObjects);
 
@@ -269,8 +271,10 @@ protected:
 	UE_API bool IsOtherActorBeingRecorded(AActor* OtherActor) const override;
 	/** Returns a valid guid if the other actor is also being recorded by the owning UTakeRecorderSources. Useful for knowing the guid of that Actor without knowing if it's a Possessable or a Spawnable. */
 	UE_API FGuid GetRecordedActorGuid(class AActor* OtherActor) const override;
-	/** Returns id of the active level sequence for the given Actor*/
-	UE_API FMovieSceneSequenceID GetLevelSequenceID(class AActor* OtherActor) override;
+	/** Returns active level sequence for the given Actor */
+	UE_API virtual ULevelSequence* GetLevelSequence(const AActor* OtherActor) override;
+	/** Returns id of the active level sequence for the given Actor */
+	UE_API virtual FMovieSceneSequenceID GetLevelSequenceID(class AActor* OtherActor) override;
 	/** Returns generic track recorder settings */
 	UE_API FTrackRecorderSettings GetTrackRecorderSettings() const override;
 	/** Returns offset that may get set when recording a skeletal mesh animation. Needed to correctly transform attached children*/
@@ -310,21 +314,24 @@ private:
 
 	/** Return either instance or global animation track recorder editor settings. */
 	UMovieSceneAnimationTrackRecorderEditorSettings* GetBestTrackRecorderEditorSettings() const;
+	
+	/** Create all pending TrackRecorders. */
+	void DoCreatePendingTrackRecorders();
 
 private:
 	/** Object Binding guid that is created in the Level Sequence when recording starts.*/
 	FGuid CachedObjectBindingGuid;
 
 	/** Array of actors that have some sort of referenced link to our actor (such as an object attached to our hierarchy) that need Actor Source recorders initialized for them. List is emptied after sources are created. */
-	TSet<class AActor*> NewReferencedActors;
+	TSet<AActor*> NewReferencedActors;
 	/** Array of Actor Sources that we ended up creating that we need to clean up when we stop recording. */
-	TArray<class UTakeRecorderSource*> AddedActorSources;
+	TArray<UTakeRecorderSource*> AddedActorSources;
 
 	/** 
 	* A pointer to the Object Template instance that was created inside the Target Level Sequence at the start of recording.
 	* Can be null (if recording to a possessable) or when we are not recording. 
 	*/
-	TWeakObjectPtr<class AActor> CachedObjectTemplate;
+	TWeakObjectPtr<AActor> CachedObjectTemplate;
 	/**
 	* A set of components that we're currently recording. We compare the components from one frame to the next to see if any
 	* components have been added or removed so we can appropriately update their Spawn tracked. Empty unless a recording is
@@ -342,6 +349,18 @@ private:
 	*/
 	FActorSerializer ActorSerializer;
 	
+	/** Args for calling CreateTrack on TrackRecorders. */
+	struct FCreateTrackArgs
+	{
+		TWeakObjectPtr<UMovieSceneTrackRecorder> TrackRecorder;
+		TWeakObjectPtr<UObject> ObjectToRecord;
+		TWeakObjectPtr<UMovieSceneTrackRecorderSettings> FactorySettings;
+		FGuid ObjectGuid;
+	};
+	
+	/** TrackRecorders we need to create after PreRecord has completed for all sources. */
+	TArray<FCreateTrackArgs> PendingTrackRecorders;
+	
 	/** If this actor source should be considered temporary, and removed later. */
 	bool bIsTemporarySource = false;
 };
diff --git a/Engine/Plugins/VirtualProduction/Takes/Source/TakeTrackRecorders/Private/MovieSceneActorReferenceTrackRecorder.cpp b/Engine/Plugins/VirtualProduction/Takes/Source/TakeTrackRecorders/Private/MovieSceneActorReferenceTrackRecorder.cpp
--- a/Engine/Plugins/VirtualProduction/Takes/Source/TakeTrackRecorders/Private/MovieSceneActorReferenceTrackRecorder.cpp
+++ b/Engine/Plugins/VirtualProduction/Takes/Source/TakeTrackRecorders/Private/MovieSceneActorReferenceTrackRecorder.cpp
@@ -266,26 +266,28 @@ FMovieSceneActorReferenceKey UMovieSceneActorReferenceTrackRecorder::ActorToRefe
 		}
 		else
 		{
-			FGuid ParentActorId = TakesUtils::ResolveActorFromSequence(InActor, RootLevelSequence);
-				
-			if (!ParentActorId.IsValid())
+			ULevelSequence* TargetLevelSequence = OwningTakeRecorderSource->GetLevelSequence(InActor);
+			if (ensure(TargetLevelSequence))
 			{
-				UE::Sequencer::FCreateBindingParams Params;
-				// Spawnable flag being accurate is critical for proper binding to actors depending on the actor & environment settings.
-				Params.bSpawnable = !OwningTakeRecorderSource->IsOtherActorPossessable(InActor);
-				if (Params.bSpawnable)
+				FGuid ParentActorId = TakesUtils::ResolveActorFromSequence(InActor, TargetLevelSequence);
+				if (!ParentActorId.IsValid())
 				{
-					Params.bAllowCustomBinding = true;
-					Params.BindingNameOverride = InActor->GetActorLabel();
+					UTakeRecorderSources* Sources = TargetLevelSequence->FindOrAddMetaData<UTakeRecorderSources>();
+					TargetLevelSequence = Sources->AddSourceForActorBinding(InActor, TargetLevelSequence);
+					check(TargetLevelSequence);
+					ParentActorId = TakesUtils::ResolveActorFromSequence(InActor, TargetLevelSequence);
+				}
+			
+				// Actor reference tracks should always be able to find a valid guid -- all sources must be established by this point.
+				if (ensure(ParentActorId.IsValid()))
+				{
+					const FMovieSceneSequenceID TargetSequenceID = OwningTakeRecorderSource->GetLevelSequenceID(InActor);
+					const FMovieSceneSequenceID ThisSequenceID   = OwningTakeRecorderSource->GetSequenceID();
+			
+					// Setup binding relative to a target sequence, which may or may not be the root sequence, so we can support references in subsequences.
+					Binding = UE::MovieScene::FRelativeObjectBindingID(ThisSequenceID, TargetSequenceID, ParentActorId, OwningTakeRecorderSource->GetRootLevelSequence());
 				}
-				ParentActorId = FSequencerUtilities::CreateOrReplaceBinding(nullptr, RootLevelSequence, InActor, Params);
-				// If this is before a recording started or the actor is already being recorded, we don't need to create the source now
-				// because the system will handle it normally instead.
-				const bool bCreateNewSourceNow = bIsRecording && !OwningTakeRecorderSource->IsOtherActorBeingRecorded(InActor);
-				UTakeRecorderSources* Sources = RootLevelSequence->FindOrAddMetaData<UTakeRecorderSources>();
-				Sources->OnNewActorBindingAdded(InActor, ParentActorId, RootLevelSequence, bCreateNewSourceNow);
 			}
-			Binding = UE::MovieScene::FRelativeObjectBindingID(ParentActorId);
 		}
 
 		if (Binding.IsValid())
diff --git a/Engine/Plugins/VirtualProduction/Takes/Source/TakeTrackRecorders/Private/MovieSceneAnimationTrackRecorder.cpp b/Engine/Plugins/VirtualProduction/Takes/Source/TakeTrackRecorders/Private/MovieSceneAnimationTrackRecorder.cpp
--- a/Engine/Plugins/VirtualProduction/Takes/Source/TakeTrackRecorders/Private/MovieSceneAnimationTrackRecorder.cpp
+++ b/Engine/Plugins/VirtualProduction/Takes/Source/TakeTrackRecorders/Private/MovieSceneAnimationTrackRecorder.cpp
@@ -240,7 +240,10 @@ void UMovieSceneAnimationTrackRecorder::FinalizeTrackImpl()
 void UMovieSceneAnimationTrackRecorder::CancelTrackImpl()
 {
 	TArray<UObject*> AssetsToCleanUp;
-	AssetsToCleanUp.Add(GetAnimSequence());
+	if (UAnimSequence* Anim = GetAnimSequence())
+	{
+		AssetsToCleanUp.Add(Anim);
+	}
 	
 	if (GEditor && AssetsToCleanUp.Num() > 0)
 	{
diff --git a/Engine/Plugins/VirtualProduction/Takes/Source/TakeTrackRecorders/Public/TrackRecorders/IMovieSceneTrackRecorderHost.h b/Engine/Plugins/VirtualProduction/Takes/Source/TakeTrackRecorders/Public/TrackRecorders/IMovieSceneTrackRecorderHost.h
--- a/Engine/Plugins/VirtualProduction/Takes/Source/TakeTrackRecorders/Public/TrackRecorders/IMovieSceneTrackRecorderHost.h
+++ b/Engine/Plugins/VirtualProduction/Takes/Source/TakeTrackRecorders/Public/TrackRecorders/IMovieSceneTrackRecorderHost.h
@@ -165,6 +165,9 @@ public:
 	*/
 	virtual FGuid GetRecordedActorGuid(class AActor* OtherActor) const = 0;
 
+	/** Get the level sequence the other actor is coming from. */
+	virtual ULevelSequence* GetLevelSequence(const AActor* OtherActor) { return nullptr; }
+	
 	/**
 	* Get the sequence id of the level sequence the other actor is coming from.
 	* Used for setting cross sequence bindings.
diff --git a/Engine/Plugins/VirtualProduction/Takes/Source/TakesCore/Private/TakeRecorderSources.cpp b/Engine/Plugins/VirtualProduction/Takes/Source/TakesCore/Private/TakeRecorderSources.cpp
--- a/Engine/Plugins/VirtualProduction/Takes/Source/TakesCore/Private/TakeRecorderSources.cpp
+++ b/Engine/Plugins/VirtualProduction/Takes/Source/TakesCore/Private/TakeRecorderSources.cpp
@@ -111,24 +111,28 @@ UTakeRecorderSources::FOnSourceRemoved& UTakeRecorderSources::OnSourceRemoved()
 	return OnSourceRemovedEvent;
 }
 
-UTakeRecorderSources::FOnActorBindingAdded& UTakeRecorderSources::GetOnActorBindingAddedDelegate()
+UTakeRecorderSources::FOnActorBindingAdded& UTakeRecorderSources::OnAddSourceForActorBinding()
 {
 	static FOnActorBindingAdded OnActorBindingAddedDelegate;
 	return OnActorBindingAddedDelegate;
 }
 
-const TMap<TObjectPtr<AActor>, FGuid>& UTakeRecorderSources::GetNewActorBindings()
+ULevelSequence* UTakeRecorderSources::AddSourceForActorBinding(AActor* InActor, ULevelSequence* InLevelSequence)
 {
-	return NewActorBindings;
-}
-
-void UTakeRecorderSources::OnNewActorBindingAdded(AActor* InActor, const FGuid& InBindingGuid, ULevelSequence* InLevelSequence, bool bCreateNewSourceNow)
-{
-	NewActorBindings.Add(InActor, InBindingGuid);
-	if (bCreateNewSourceNow)
+	// Communicate to the actor sources module we're adding a new binding.
+	check(OnAddSourceForActorBinding().IsBound());
+	UTakeRecorderSource* NewSource = OnAddSourceForActorBinding().Execute(InActor, InLevelSequence);
+	if (ensure(NewSource))
 	{
-		GetOnActorBindingAddedDelegate().Broadcast(InActor, InLevelSequence);
+		// Look for a subsequence if applicable.
+		if (TObjectPtr<ULevelSequence>* Subsequence = SourceSubSequenceMap.Find(NewSource))
+		{
+			return *Subsequence;
+		}
 	}
+
+	// Fallback to original sequence (probably root).
+	return InLevelSequence;
 }
 
 void UTakeRecorderSources::PostEditChangeProperty(FPropertyChangedEvent& PropertyChangedEvent)
@@ -159,8 +163,16 @@ void UTakeRecorderSources::StartPreRecordingRecursive(const TArray<UTakeRecorder
 	TSet<FString> UsedTrackNames;
 	PreRecordingRecursive(InSources, InRootSequence, NewSourcesOut, InManifestSerializer, UsedTrackNames);
 	
-	// Make sure any recently added bindings are cleared out in case there wasn't a source to clear them.
-	NewActorBindings.Empty();
+	TSet<UTakeRecorderSource*> AllSources(InSources);
+	AllSources.Append(NewSourcesOut);
+	
+	for (UTakeRecorderSource* Source : AllSources)
+	{
+		if (Source)
+		{
+			Source->OnAllSourcesPreRecordingFinished();
+		}
+	}
 }
 
 void UTakeRecorderSources::PreRecordingRecursive(TArray<UTakeRecorderSource*> InSources, ULevelSequence* InRootSequence,
@@ -340,38 +352,45 @@ void UTakeRecorderSources::StartRecordingPreRecordedSources(const FQualifiedFram
 	PreRecordedSources.Reset(0);
 }
 
-void UTakeRecorderSources::PreRecordSources(TArray<UTakeRecorderSource *> InSources)
+void UTakeRecorderSources::PreRecordSources(TArray<UTakeRecorderSource*> InSources)
 {
 	PreRecordedSources.Reset(0);
 	StartPreRecordingRecursive(InSources, CachedLevelSequence, PreRecordedSources, CachedManifestSerializer);
 }
 
 
-void UTakeRecorderSources::StartRecordingSource(TArray<UTakeRecorderSource *> InSources, const FQualifiedFrameTime& CurrentFrameTime)
+void UTakeRecorderSources::StartRecordingSource(TArray<UTakeRecorderSource*> InSources, const FQualifiedFrameTime& CurrentFrameTime)
 {
 	// This calls PreRecording recursively on every source so that all sources that get added by another source
 	// have had PreRecording called.
-	TArray<UTakeRecorderSource *> NewSources;
+	TArray<UTakeRecorderSource*> NewSources;
 	StartPreRecordingRecursive(InSources, CachedLevelSequence, NewSources, CachedManifestSerializer);
-	if (NewSources.Num() > 0)
+	
+	// Only start recording these sources if a recording is already in progress. It's possible this method is called during PreRecord, and we don't
+	// want to end up calling Record twice. This condition helps when starting a recording with an actor reference property where the value points
+	// to another actor that has not been added as a source, and we need to add it automatically before formally starting the recording.
+	if (bIsRecording)
 	{
-		InSources.Append(NewSources);
+		// Ensures uniqueness. Recursive method will add old as new sources.
+		TSet<UTakeRecorderSource*> UniqueSources(InSources);
+		UniqueSources.Append(NewSources);
+		InSources = UniqueSources.Array();
+		StartRecordingTheseSources(InSources, CurrentFrameTime);
 	}
-	StartRecordingTheseSources(InSources, CurrentFrameTime);
 }
 
 
-void UTakeRecorderSources::StartRecordingTheseSources(const TArray<UTakeRecorderSource *>& InSources, const FQualifiedFrameTime& CurrentFrameTime)
+void UTakeRecorderSources::StartRecordingTheseSources(const TArray<UTakeRecorderSource*>& InSources, const FQualifiedFrameTime& CurrentFrameTime)
 {
-	FTimecode CurrentTimecode = FTimecode::FromFrameNumber(FFrameRate::TransformTime(CurrentFrameTime.Time, TargetLevelSequenceTickResolution, TargetLevelSequenceDisplayRate).FloorToFrame(), TargetLevelSequenceDisplayRate);
+	const FTimecode CurrentTimecode = FTimecode::FromFrameNumber(FFrameRate::TransformTime(CurrentFrameTime.Time, TargetLevelSequenceTickResolution, TargetLevelSequenceDisplayRate).FloorToFrame(), TargetLevelSequenceDisplayRate);
 
-	for (auto Source : InSources)
+	for (UTakeRecorderSource* Source : InSources)
 	{
 		if (Source && Source->bEnabled)
 		{
 			if (Settings.bRecordSourcesIntoSubSequences && Source->SupportsSubscenes()) //Set Timcode on MovieScene if we created a sub scene for it
 			{
-				for (UMovieSceneSubSection* ActiveSubSection : ActiveSubSections)
+				for (const TObjectPtr<UMovieSceneSubSection>& ActiveSubSection : ActiveSubSections)
 				{
 					// Set timecode source and start time if it hasn't been set
 					if (ActiveSubSection->TimecodeSource.Timecode == FTimecode())
@@ -624,7 +643,6 @@ void UTakeRecorderSources::StopRecording(class ULevelSequence* InSequence, const
 	SourceSubSequenceMap.Empty();
 	ActiveSubSections.Empty();
 	CreatedManifestSerializers.Empty();
-	NewActorBindings.Empty();
 	CachedManifestSerializer = nullptr;
 	CachedLevelSequence = nullptr;
 
diff --git a/Engine/Plugins/VirtualProduction/Takes/Source/TakesCore/Public/TakeRecorderSource.h b/Engine/Plugins/VirtualProduction/Takes/Source/TakesCore/Public/TakeRecorderSource.h
--- a/Engine/Plugins/VirtualProduction/Takes/Source/TakesCore/Public/TakeRecorderSource.h
+++ b/Engine/Plugins/VirtualProduction/Takes/Source/TakesCore/Public/TakeRecorderSource.h
@@ -80,6 +80,9 @@ public:
 	*/
 	UE_API virtual TArray<UTakeRecorderSource*> PreRecording(ULevelSequence* InSequence, FMovieSceneSequenceID InSequenceID, ULevelSequence* InRootSequence, FManifestSerializer* InManifestSerializer);
 
+	/** Called after all sources have had PreRecording run. */
+	UE_API virtual void OnAllSourcesPreRecordingFinished() { }
+	
 	/**
 	* This is called when the UTakeRecorderSources starts a recording, after all sources have had PreRecording called on them.
 	* Implementations should avoid blocking on this call (instead place that in PreRecording) so that the sources all get
diff --git a/Engine/Plugins/VirtualProduction/Takes/Source/TakesCore/Public/TakeRecorderSources.h b/Engine/Plugins/VirtualProduction/Takes/Source/TakesCore/Public/TakeRecorderSources.h
--- a/Engine/Plugins/VirtualProduction/Takes/Source/TakesCore/Public/TakeRecorderSources.h
+++ b/Engine/Plugins/VirtualProduction/Takes/Source/TakesCore/Public/TakeRecorderSources.h
@@ -35,6 +35,12 @@ class ULevelSequence;
 class FViewport;
 class FCanvas;
 
+namespace UE::TakeRecorderSources
+{
+	class FTakeRecorderSourcesModule;	
+}
+	
+
 /**
  * A list of sources to record for any given take. Stored as meta-data on ULevelSequence through ULevelSequence::FindMetaData<UTakeRecorderSources>
  */
@@ -137,27 +143,21 @@ public:
 	/** Multicast delegate for when any source is removed. */
 	static UE_API FOnSourceRemoved& OnSourceRemoved();
 	
-	DECLARE_MULTICAST_DELEGATE_TwoParams(FOnActorBindingAdded, AActor*, ULevelSequence*);
-	/** Multicast delegate for when an actor binding is added. */
-	static UE_API FOnActorBindingAdded& GetOnActorBindingAddedDelegate();
-
-	/** Actors that have had bindings created recently. */
-	UE_API const TMap<TObjectPtr<AActor>, FGuid>& GetNewActorBindings();
+private:
+	friend class UE::TakeRecorderSources::FTakeRecorderSourcesModule;
 	
+	DECLARE_DELEGATE_RetVal_TwoParams(UTakeRecorderSource*, FOnActorBindingAdded, AActor*, ULevelSequence*);
+	/** Delegate for when a source is being added for an actor binding. This is for our Sources Module, so it can properly add an actor source. */
+	static UE_API FOnActorBindingAdded& OnAddSourceForActorBinding();
+	
+public:
 	/**
 	 * Notifies sources that a new actor binding has been added.
 	 * 
 	 * @param InActor The actor the binding was added for.
-	 * @param InBindingGuid The binding guid created.
 	 * @param InLevelSequence The level sequence of the binding.
-	 * @param bCreateNewSourceNow If a new source should be created for the actor this moment. This should be true
-	 * if a recording is already in progress.
 	 */
-	UE_API void OnNewActorBindingAdded(AActor* InActor, const FGuid& InBindingGuid, ULevelSequence* InLevelSequence, bool bCreateNewSourceNow);
-
-private:
-	UPROPERTY(Transient)
-	TMap<TObjectPtr<AActor>, FGuid> NewActorBindings;
+	UE_API ULevelSequence* AddSourceForActorBinding(AActor* InActor, ULevelSequence* InLevelSequence);
 	
 public:
 
diff --git a/Engine/Source/Editor/Sequencer/Private/SequencerUtilities.cpp b/Engine/Source/Editor/Sequencer/Private/SequencerUtilities.cpp
--- a/Engine/Source/Editor/Sequencer/Private/SequencerUtilities.cpp
+++ b/Engine/Source/Editor/Sequencer/Private/SequencerUtilities.cpp
@@ -3134,8 +3134,7 @@ FGuid CreateGenericBinding(TSharedPtr<ISequencer> Sequencer, UMovieSceneSequence
 
 		// If we're not trying to replace a binding, and the object is a spawnable, try and bind to that first
 		if (InParams.ReplacementGuid.IsValid()
-			|| (Sequencer && (!CurrentObject || !NewPossessable->BindSpawnableObject(Sequencer->GetFocusedTemplateID(), CurrentObject, Sequencer->GetSharedPlaybackState())))
-			|| (!Sequencer && !InParams.bSpawnable)) /** If binding without sequencer we still need to call AddBinding */
+			|| (Sequencer && (!CurrentObject || !NewPossessable->BindSpawnableObject(Sequencer->GetFocusedTemplateID(), CurrentObject, Sequencer->GetSharedPlaybackState()))))
 		{
 			FUniversalObjectLocator Locator;
 			if (CurrentObject && (!OwnerSequence->MakeLocatorForObject(CurrentObject, Context, Locator) || Locator.IsEmpty()))
