TakeRecorder: Add TakeRecorderSourceModified event to the TakeRecorderSubsystem.
Fix ActorSources not calling UObject::PostEditChangeProperty.

#Synced-CL 50935209

diff --git a/Engine/Plugins/VirtualProduction/Takes/Source/TakeRecorder/Private/Recorder/TakeRecorderSubsystem.cpp b/Engine/Plugins/VirtualProduction/Takes/Source/TakeRecorder/Private/Recorder/TakeRecorderSubsystem.cpp
--- a/Engine/Plugins/VirtualProduction/Takes/Source/TakeRecorder/Private/Recorder/TakeRecorderSubsystem.cpp
+++ b/Engine/Plugins/VirtualProduction/Takes/Source/TakeRecorder/Private/Recorder/TakeRecorderSubsystem.cpp
@@ -325,4 +325,9 @@ UTakeRecorderSources::FOnSourceRemoved& UTakeRecorderSubsystem::GetOnRecordingSo
 	return OnRecordingSourceRemovedEvent;
 }
 
+UTakeRecorderSources::FOnSourceModified& UTakeRecorderSubsystem::GetOnRecordingSourceModifiedEvent()
+{
+	return OnRecordingSourceModifiedEvent;
+}
+
 #undef LOCTEXT_NAMESPACE
diff --git a/Engine/Plugins/VirtualProduction/Takes/Source/TakeRecorder/Private/Recorder/TakeRecorderSubsystemImplementation.cpp b/Engine/Plugins/VirtualProduction/Takes/Source/TakeRecorder/Private/Recorder/TakeRecorderSubsystemImplementation.cpp
--- a/Engine/Plugins/VirtualProduction/Takes/Source/TakeRecorder/Private/Recorder/TakeRecorderSubsystemImplementation.cpp
+++ b/Engine/Plugins/VirtualProduction/Takes/Source/TakeRecorder/Private/Recorder/TakeRecorderSubsystemImplementation.cpp
@@ -85,6 +85,7 @@ void UTakeRecorderSubsystemImplementation::InitializeImplementation(UTakeRecorde
 
 	UTakeRecorderSources::OnSourceAdded().AddUObject(this, &UTakeRecorderSubsystemImplementation::OnSourceAdded);
 	UTakeRecorderSources::OnSourceRemoved().AddUObject(this, &UTakeRecorderSubsystemImplementation::OnSourceRemoved);
+	UTakeRecorderSources::OnSourceModified().AddUObject(this, &UTakeRecorderSubsystemImplementation::OnSourceModified);
 
 	UTakeMetaData::OnTakeSlateChanged().AddUObject(this, &UTakeRecorderSubsystemImplementation::OnTakeSlateChanged);
 	UTakeMetaData::OnTakeNumberChanged().AddUObject(this, &UTakeRecorderSubsystemImplementation::OnTakeNumberChanged);
@@ -104,6 +105,7 @@ void UTakeRecorderSubsystemImplementation::DeinitializeImplementation()
 
 	UTakeRecorderSources::OnSourceAdded().RemoveAll(this);
 	UTakeRecorderSources::OnSourceRemoved().RemoveAll(this);
+	UTakeRecorderSources::OnSourceModified().RemoveAll(this);
 
 	UTakeMetaData::OnTakeSlateChanged().RemoveAll(this);
 	UTakeMetaData::OnTakeNumberChanged().RemoveAll(this);
@@ -1060,6 +1062,12 @@ void UTakeRecorderSubsystemImplementation::OnSourceRemoved(UTakeRecorderSource*
 	WITH_OWNING_SUBSYSTEM(TakeRecorderSourceRemoved.Broadcast(Source));
 }
 
+void UTakeRecorderSubsystemImplementation::OnSourceModified(UTakeRecorderSource* Source)
+{
+	WITH_OWNING_SUBSYSTEM(GetOnRecordingSourceModifiedEvent().Broadcast(Source));
+	WITH_OWNING_SUBSYSTEM(TakeRecorderSourceModified.Broadcast(Source));
+}
+
 void UTakeRecorderSubsystemImplementation::BindToNamingTokenEvents()
 {
 	if (GEngine && !GetNamingTokensData()->TakeRecorderNamingTokens.IsValid())
diff --git a/Engine/Plugins/VirtualProduction/Takes/Source/TakeRecorder/Private/Recorder/TakeRecorderSubsystemImplementation.h b/Engine/Plugins/VirtualProduction/Takes/Source/TakeRecorder/Private/Recorder/TakeRecorderSubsystemImplementation.h
--- a/Engine/Plugins/VirtualProduction/Takes/Source/TakeRecorder/Private/Recorder/TakeRecorderSubsystemImplementation.h
+++ b/Engine/Plugins/VirtualProduction/Takes/Source/TakeRecorder/Private/Recorder/TakeRecorderSubsystemImplementation.h
@@ -125,6 +125,9 @@ private:
 
 	/** Callback when any source is removed. */
 	void OnSourceRemoved(UTakeRecorderSource* Source);
+	
+	/** Callback when any sources is modified. */
+	void OnSourceModified(UTakeRecorderSource* Source);
 
 	/** Setup bindings to our naming tokens. */
 	void BindToNamingTokenEvents();
diff --git a/Engine/Plugins/VirtualProduction/Takes/Source/TakeRecorder/Public/Recorder/TakeRecorderSubsystem.h b/Engine/Plugins/VirtualProduction/Takes/Source/TakeRecorder/Public/Recorder/TakeRecorderSubsystem.h
--- a/Engine/Plugins/VirtualProduction/Takes/Source/TakeRecorder/Public/Recorder/TakeRecorderSubsystem.h
+++ b/Engine/Plugins/VirtualProduction/Takes/Source/TakeRecorder/Public/Recorder/TakeRecorderSubsystem.h
@@ -33,6 +33,7 @@ DECLARE_DYNAMIC_MULTICAST_DELEGATE_TwoParams(FTakeRecorderTakeNumberChanged, int
  */
 DECLARE_DYNAMIC_MULTICAST_DELEGATE_OneParam(FOnTakeRecorderSourceAddedDynamic, UTakeRecorderSource*, Source);
 DECLARE_DYNAMIC_MULTICAST_DELEGATE_OneParam(FOnTakeRecorderSourceRemovedDynamic, UTakeRecorderSource*, Source);
+DECLARE_DYNAMIC_MULTICAST_DELEGATE_OneParam(FOnTakeRecorderSourceModifiedDynamic, UTakeRecorderSource*, Source);
 
 /**
 * UTakeRecorderSubsystem Subsystem for Take Recorder.
@@ -304,6 +305,7 @@ public:
 	/** Set the target record class of the Take Preset. */
 	UFUNCTION(BlueprintCallable, Category = TakeRecorder)
 	UE_API virtual void SetPresetTargetRecordClass(const TSubclassOf<ULevelSequence> InClass) override;
+	// ~End ITakeRecorderSubsystem interface
 	
 	/**
 	 * Retrieve a multi-cast delegate that is triggered when a recording pre-initializes
@@ -337,8 +339,10 @@ public:
 	 * Retrieve a multi-cast delegate that is triggered when a source is removed.
 	 */
 	UE_API UTakeRecorderSources::FOnSourceRemoved& GetOnRecordingSourceRemovedEvent();
-
-	// ~End ITakeRecorderSubsystem interface
+	/**
+	 * Retrieve a multi-cast delegate that is triggered when a source is modified.
+	 */
+	UE_API UTakeRecorderSources::FOnSourceModified& GetOnRecordingSourceModifiedEvent();
 	
 private:
 	// Friend class declared as a precaution in case the implementation needs to access private members.
@@ -370,6 +374,8 @@ private:
 	UTakeRecorderSources::FOnSourceAdded OnRecordingSourceAddedEvent;
 	/** Multicast delegate when any source has been added. For native binding. */
 	UTakeRecorderSources::FOnSourceRemoved OnRecordingSourceRemovedEvent;
+	/** Multicast delegate when any source has been modified. For native binding. */
+	UTakeRecorderSources::FOnSourceModified OnRecordingSourceModifiedEvent;
 
 public:
 	/** Called before initialization occurs (ie. when the recording button is pressed and before the countdown starts) */
@@ -415,6 +421,10 @@ public:
 	/** Multicast delegate when any source has been removed. */
 	UPROPERTY(BlueprintAssignable, Category=TakeRecorder)
 	FOnTakeRecorderSourceRemovedDynamic TakeRecorderSourceRemoved;
+	
+	/** Multicast delegate when any source has been modified. */
+	UPROPERTY(BlueprintAssignable, Category=TakeRecorder)
+	FOnTakeRecorderSourceModifiedDynamic TakeRecorderSourceModified;
 };
 
 #undef UE_API
diff --git a/Engine/Plugins/VirtualProduction/Takes/Source/TakeRecorderSources/Private/TakeRecorderActorSource.cpp b/Engine/Plugins/VirtualProduction/Takes/Source/TakeRecorderSources/Private/TakeRecorderActorSource.cpp
--- a/Engine/Plugins/VirtualProduction/Takes/Source/TakeRecorderSources/Private/TakeRecorderActorSource.cpp
+++ b/Engine/Plugins/VirtualProduction/Takes/Source/TakeRecorderSources/Private/TakeRecorderActorSource.cpp
@@ -1173,6 +1173,8 @@ bool UTakeRecorderActorSource::EnsureObjectTemplateHasComponent(UActorComponent*
 
 void UTakeRecorderActorSource::PostEditChangeProperty(FPropertyChangedEvent& PropertyChangedEvent)
 {
+	Super::PostEditChangeProperty(PropertyChangedEvent);
+
 	if (PropertyChangedEvent.Property && PropertyChangedEvent.Property->GetFName() == GET_MEMBER_NAME_CHECKED(UTakeRecorderActorSource, Target))
 	{
 		TrackTint = FColor(67, 148, 135);
diff --git a/Engine/Plugins/VirtualProduction/Takes/Source/TakesCore/Private/TakeRecorderSource.cpp b/Engine/Plugins/VirtualProduction/Takes/Source/TakesCore/Private/TakeRecorderSource.cpp
--- a/Engine/Plugins/VirtualProduction/Takes/Source/TakesCore/Private/TakeRecorderSource.cpp
+++ b/Engine/Plugins/VirtualProduction/Takes/Source/TakesCore/Private/TakeRecorderSource.cpp
@@ -14,6 +14,15 @@ UTakeRecorderSource::UTakeRecorderSource(const FObjectInitializer& ObjInit)
 	TrackTint = FColor(127, 127, 127);
 }
 
+#if WITH_EDITOR
+void UTakeRecorderSource::PostEditChangeProperty(FPropertyChangedEvent& PropertyChangedEvent)
+{
+	Super::PostEditChangeProperty(PropertyChangedEvent);
+	
+	UTakeRecorderSources::OnSourceModified().Broadcast(this);
+}
+#endif
+
 const FSlateBrush* UTakeRecorderSource::GetDisplayIconImpl() const
 {
 	return FSlateIconFinder::FindCustomIconBrushForClass(GetClass(), TEXT("ClassThumbnail"));
diff --git a/Engine/Plugins/VirtualProduction/Takes/Source/TakesCore/Private/TakeRecorderSources.cpp b/Engine/Plugins/VirtualProduction/Takes/Source/TakesCore/Private/TakeRecorderSources.cpp
--- a/Engine/Plugins/VirtualProduction/Takes/Source/TakesCore/Private/TakeRecorderSources.cpp
+++ b/Engine/Plugins/VirtualProduction/Takes/Source/TakesCore/Private/TakeRecorderSources.cpp
@@ -111,6 +111,12 @@ UTakeRecorderSources::FOnSourceRemoved& UTakeRecorderSources::OnSourceRemoved()
 	return OnSourceRemovedEvent;
 }
 
+UTakeRecorderSources::FOnSourceModified& UTakeRecorderSources::OnSourceModified()
+{
+	static FOnSourceModified OnSourceModifiedEvent;
+	return OnSourceModifiedEvent;
+}
+
 UTakeRecorderSources::FOnActorBindingAdded& UTakeRecorderSources::OnAddSourceForActorBinding()
 {
 	static FOnActorBindingAdded OnActorBindingAddedDelegate;
diff --git a/Engine/Plugins/VirtualProduction/Takes/Source/TakesCore/Public/TakeRecorderSource.h b/Engine/Plugins/VirtualProduction/Takes/Source/TakesCore/Public/TakeRecorderSource.h
--- a/Engine/Plugins/VirtualProduction/Takes/Source/TakesCore/Public/TakeRecorderSource.h
+++ b/Engine/Plugins/VirtualProduction/Takes/Source/TakesCore/Public/TakeRecorderSource.h
@@ -34,6 +34,12 @@ public:
 	GENERATED_BODY()
 
 	UE_API UTakeRecorderSource(const FObjectInitializer& ObjInit);
+	
+	// ~Begin UObject
+#if WITH_EDITOR
+	UE_API virtual void PostEditChangeProperty(FPropertyChangedEvent& PropertyChangedEvent) override;
+#endif
+	// ~End UObject
 
 	/** True if this source is cued for recording or not */
 	UPROPERTY(BlueprintReadWrite, Category="Source")
diff --git a/Engine/Plugins/VirtualProduction/Takes/Source/TakesCore/Public/TakeRecorderSources.h b/Engine/Plugins/VirtualProduction/Takes/Source/TakesCore/Public/TakeRecorderSources.h
--- a/Engine/Plugins/VirtualProduction/Takes/Source/TakesCore/Public/TakeRecorderSources.h
+++ b/Engine/Plugins/VirtualProduction/Takes/Source/TakesCore/Public/TakeRecorderSources.h
@@ -136,6 +136,7 @@ public:
 	
 	DECLARE_MULTICAST_DELEGATE_OneParam(FOnSourceAdded, UTakeRecorderSource* Source);
 	DECLARE_MULTICAST_DELEGATE_OneParam(FOnSourceRemoved, UTakeRecorderSource* Source);
+	DECLARE_MULTICAST_DELEGATE_OneParam(FOnSourceModified, UTakeRecorderSource* Source);
 	
 	/** Multicast delegate for when any source is added. */
 	static UE_API FOnSourceAdded& OnSourceAdded();
@@ -143,6 +144,9 @@ public:
 	/** Multicast delegate for when any source is removed. */
 	static UE_API FOnSourceRemoved& OnSourceRemoved();
 	
+	/** Multicast delegate for when any source is modified. */
+	static UE_API FOnSourceModified& OnSourceModified();
+	
 private:
 	friend class UE::TakeRecorderSources::FTakeRecorderSourcesModule;
 	
