TakeRecorder: Fix enabling or disabling sources in the UI not firing PostEditChangeProperty.
Follow up to 50943023.

#Synced-CL 50935209

diff --git a/Engine/Plugins/VirtualProduction/Takes/Source/TakeRecorder/Private/Widgets/STakeRecorderSources.cpp b/Engine/Plugins/VirtualProduction/Takes/Source/TakeRecorder/Private/Widgets/STakeRecorderSources.cpp
--- a/Engine/Plugins/VirtualProduction/Takes/Source/TakeRecorder/Private/Widgets/STakeRecorderSources.cpp
+++ b/Engine/Plugins/VirtualProduction/Takes/Source/TakeRecorder/Private/Widgets/STakeRecorderSources.cpp
@@ -241,12 +241,17 @@ private:
 				? LOCTEXT("EnableSources", "Enable Recording {0}|plural(one=Source, other=Sources)")
 				: LOCTEXT("DisableSources", "Disable Recording {0}|plural(one=Source, other=Sources)");
 
+			FProperty* EnabledProperty = UTakeRecorderSource::StaticClass()->FindPropertyByName(GET_MEMBER_NAME_CHECKED(UTakeRecorderSource, bEnabled));
+			
 			if (!SelectedSources.Contains(ThisSource))
 			{
 				FScopedTransaction Transaction(FText::Format(TransactionFormat, 1));
 
 				ThisSource->Modify();
 				ThisSource->bEnabled = bEnable;
+				
+				FPropertyChangedEvent PropertyChangedEvent(EnabledProperty, EPropertyChangeType::ValueSet);
+				ThisSource->PostEditChangeProperty(PropertyChangedEvent);
 			}
 			else 
 			{
@@ -256,6 +261,9 @@ private:
 				{
 					SelectedSource->Modify();
 					SelectedSource->bEnabled = bEnable;
+					
+					FPropertyChangedEvent PropertyChangedEvent(EnabledProperty, EPropertyChangeType::ValueSet);
+					SelectedSource->PostEditChangeProperty(PropertyChangedEvent);
 				}
 			}
 		}
