TakeRecorder: Fix possessable actors lose their binding when recorded in Take Recorder when "Record to Possessable" is unchecked.
When cleaning external resources Take Recorder takes into account whether individual sources are possessable or spawnable.

Fix ensure that could be tripped opening the Add Source dropdown.

#Synced-CL 51303076

diff --git a/Engine/Plugins/VirtualProduction/Takes/Source/TakeRecorder/Private/Recorder/TakeRecorder.cpp b/Engine/Plugins/VirtualProduction/Takes/Source/TakeRecorder/Private/Recorder/TakeRecorder.cpp
--- a/Engine/Plugins/VirtualProduction/Takes/Source/TakeRecorder/Private/Recorder/TakeRecorder.cpp
+++ b/Engine/Plugins/VirtualProduction/Takes/Source/TakeRecorder/Private/Recorder/TakeRecorder.cpp
@@ -1254,11 +1254,8 @@ void UTakeRecorder::StopInternal(const bool bCancelled)
 			{
 				if (Sequencer)
 				{
-					if (!Parameters.Project.bRecordToPossessable)
-					{
-						// Ensure we don't have any external references prior to saving.
-						TakesUtils::CleanupExternalReferences(SequenceAsset);
-					}
+					// Ensure we don't have any external references prior to saving. This automatically handles possessables vs spawnables.
+					TakesUtils::CleanupExternalReferences(SequenceAsset);
 					
 					// Sequencer has a feature for scrubbing to a particular frame and capturing it as thumbnail.
 					// In 5.5, this was added but hidden behind a CVar for stability reason.
@@ -1270,7 +1267,7 @@ void UTakeRecorder::StopInternal(const bool bCancelled)
 				else
 				{
 					// Fallback to "normal" saving if sequencer is not available (seems like it should be though?). Asset won't have any thumbnail.
-					TakesUtils::SaveAsset(SequenceAsset, Parameters.Project.bRecordToPossessable);
+					TakesUtils::SaveAsset(SequenceAsset);
 				}
 			}
 
diff --git a/Engine/Plugins/VirtualProduction/Takes/Source/TakeRecorderSources/Private/TakeRecorderActorSource.cpp b/Engine/Plugins/VirtualProduction/Takes/Source/TakeRecorderSources/Private/TakeRecorderActorSource.cpp
--- a/Engine/Plugins/VirtualProduction/Takes/Source/TakeRecorderSources/Private/TakeRecorderActorSource.cpp
+++ b/Engine/Plugins/VirtualProduction/Takes/Source/TakeRecorderSources/Private/TakeRecorderActorSource.cpp
@@ -1834,6 +1834,21 @@ FTrackRecorderSettings UTakeRecorderActorSource::GetTrackRecorderSettings() cons
 	return TrackRecorderSettings;
 }
 
+const IMovieSceneTrackRecorderHost* UTakeRecorderActorSource::AsTrackRecorderHost() const
+{
+	return this;
+}
+
+bool UTakeRecorderActorSource::IsActorSourceForActor(const AActor* Actor) const
+{
+#if WITH_EDITOR
+	return Actor && Target.LoadSynchronous() && Actor->GetActorLabel() == Target->GetActorLabel();
+#else
+	// branch untested
+	return Actor == Target;
+#endif
+}
+
 void UTakeRecorderActorSource::SetSourceActor(TSoftObjectPtr<AActor> InTarget)
 {
 	Target = InTarget;
diff --git a/Engine/Plugins/VirtualProduction/Takes/Source/TakeRecorderSources/Private/TakeRecorderMicrophoneAudioSource.cpp b/Engine/Plugins/VirtualProduction/Takes/Source/TakeRecorderSources/Private/TakeRecorderMicrophoneAudioSource.cpp
--- a/Engine/Plugins/VirtualProduction/Takes/Source/TakeRecorderSources/Private/TakeRecorderMicrophoneAudioSource.cpp
+++ b/Engine/Plugins/VirtualProduction/Takes/Source/TakeRecorderSources/Private/TakeRecorderMicrophoneAudioSource.cpp
@@ -424,7 +424,8 @@ TArray<UTakeRecorderSource*> UTakeRecorderMicrophoneAudioSource::PostRecording(U
 
 				if ((Sources->GetSettings().bSaveRecordedAssets) || GEditor == nullptr)
 				{
-					TakesUtils::SaveAsset(SoundWave, Sources->GetSettings().bRecordToPossessable);
+					const bool bIsPossessable = TakesUtils::IsSourcePossessable(this);
+					TakesUtils::SaveAsset(SoundWave, bIsPossessable);
 				}
 			}
 		}
diff --git a/Engine/Plugins/VirtualProduction/Takes/Source/TakeRecorderSources/Public/TakeRecorderActorSource.h b/Engine/Plugins/VirtualProduction/Takes/Source/TakeRecorderSources/Public/TakeRecorderActorSource.h
--- a/Engine/Plugins/VirtualProduction/Takes/Source/TakeRecorderSources/Public/TakeRecorderActorSource.h
+++ b/Engine/Plugins/VirtualProduction/Takes/Source/TakeRecorderSources/Public/TakeRecorderActorSource.h
@@ -155,7 +155,8 @@ public:
 	UE_API virtual FString GetSubsceneAssetName(ULevelSequence* InSequence) const override;
 	UE_API virtual void AddContentsToFolder(class UMovieSceneFolder* InFolder) override;
 	virtual FMovieSceneSequenceID GetSequenceID() const override { return TargetSequenceID; }
-
+	UE_API virtual const IMovieSceneTrackRecorderHost* AsTrackRecorderHost() const override;
+	UE_API virtual bool IsActorSourceForActor(const AActor* Actor) const override;
 	// ~UTakeRecorderSource Interface
 
 	/** Set the Target actor that we are going to record. Will reset the Recorded Property Map to defaults. */
diff --git a/Engine/Plugins/VirtualProduction/Takes/Source/TakesCore/Private/TakeRecorderSources.cpp b/Engine/Plugins/VirtualProduction/Takes/Source/TakesCore/Private/TakeRecorderSources.cpp
--- a/Engine/Plugins/VirtualProduction/Takes/Source/TakesCore/Private/TakeRecorderSources.cpp
+++ b/Engine/Plugins/VirtualProduction/Takes/Source/TakesCore/Private/TakeRecorderSources.cpp
@@ -2,34 +2,29 @@
 
 #include "TakeRecorderSources.h"
 #include "AssetRegistry/AssetData.h"
-#include "TakeRecorderSource.h"
-#include "TakesCoreLog.h"
-#include "LevelSequence.h"
-#include "TakeMetaData.h"
 #include "AssetRegistry/AssetRegistryModule.h"
+#include "CanvasTypes.h"
+#include "Editor.h"
+#include "Engine/Engine.h"
+#include "Engine/Font.h"
+#include "Engine/TimecodeProvider.h"
+#include "LevelSequence.h"
+#include "MetaData/MovieSceneShotMetaData.h"
+#include "Misc/App.h"
 #include "Modules/ModuleManager.h"
-#include "AssetRegistry/IAssetRegistry.h"
-#include "Misc/PackageName.h"
-#include "TakesUtils.h"
-#include "Tracks/MovieSceneSubTrack.h"
-#include "Sections/MovieSceneSubSection.h"
 #include "MovieScene.h"
 #include "MovieSceneFolder.h"
-#include "MovieSceneTimeHelpers.h"
-#include "ObjectTools.h"
-#include "Engine/TimecodeProvider.h"
-#include "Misc/App.h"
-#include "Editor.h"
-#include "MetaData/MovieSceneShotMetaData.h"
 #include "NamingTokens/TakeRecorderNamingTokensContext.h"
-#include "UObject/Package.h"
-
-#include "Engine/Engine.h"
-#include "Engine/Font.h"
-#include "CanvasTypes.h"
-
+#include "ObjectTools.h"
+#include "Sections/MovieSceneSubSection.h"
+#include "TakeMetaData.h"
+#include "TakeRecorderSource.h"
+#include "TakesCoreLog.h"
+#include "TakesUtils.h"
+#include "TrackRecorders/IMovieSceneTrackRecorderHost.h"
+#include "Tracks/MovieSceneSubTrack.h"
 
-#define LOCTEXT_NAMESPACE "TakesUtils"
+#define LOCTEXT_NAMESPACE "TakeRecorderSources"
 
 #include UE_INLINE_GENERATED_CPP_BY_NAME(TakeRecorderSources)
 
@@ -130,6 +125,12 @@ ULevelSequence* UTakeRecorderSources::AddSourceForActorBinding(AActor* InActor,
 	UTakeRecorderSource* NewSource = OnAddSourceForActorBinding().Execute(InActor, InLevelSequence);
 	if (ensure(NewSource))
 	{
+		if (const IMovieSceneTrackRecorderHost* AsTrackRecorderHost = NewSource->AsTrackRecorderHost())
+		{
+			const bool bIsPossessable = AsTrackRecorderHost->GetTrackRecorderSettings().bRecordToPossessable;
+			AddedActorPossessableBindings.Add(InActor, bIsPossessable);
+		}
+		
 		// Look for a subsequence if applicable.
 		if (TObjectPtr<ULevelSequence>* Subsequence = SourceSubSequenceMap.Find(NewSource))
 		{
@@ -141,6 +142,26 @@ ULevelSequence* UTakeRecorderSources::AddSourceForActorBinding(AActor* InActor,
 	return InLevelSequence;
 }
 
+bool UTakeRecorderSources::IsActorPossessable(const AActor* InActor) const
+{
+	// Check bindings first -- we can't just check the sources, because reference bindings create temporary sources that may be
+	// cleared out by this point.
+	if (const bool *AddedBinding = AddedActorPossessableBindings.Find(InActor))
+	{
+		return *AddedBinding;
+	}
+	
+	for (const UTakeRecorderSource* Source : GetSources())
+	{
+		if (Source && Source->IsActorSourceForActor(InActor))
+		{
+			return TakesUtils::IsSourcePossessable(Source);
+		}
+	}
+				
+	return GetSettings().bRecordToPossessable;
+}
+
 void UTakeRecorderSources::PostEditChangeProperty(FPropertyChangedEvent& PropertyChangedEvent)
 {
 	Super::PostEditChangeProperty(PropertyChangedEvent);
@@ -166,6 +187,8 @@ void UTakeRecorderSources::PostEditChangeProperty(FPropertyChangedEvent& Propert
 void UTakeRecorderSources::StartPreRecordingRecursive(const TArray<UTakeRecorderSource*>& InSources, ULevelSequence* InRootSequence,
 	TArray<UTakeRecorderSource*>& NewSourcesOut, FManifestSerializer* InManifestSerializer)
 {
+	AddedActorPossessableBindings.Empty();
+	
 	TSet<FString> UsedTrackNames;
 	PreRecordingRecursive(InSources, InRootSequence, NewSourcesOut, InManifestSerializer, UsedTrackNames);
 	
@@ -629,7 +652,7 @@ void UTakeRecorderSources::StopRecording(class ULevelSequence* InSequence, const
 	}
 
 	TArray<UObject*> AssetsToCleanUp;
-	for (auto SourceSubSequence : SourceSubSequenceMap)
+	for (const TTuple<TObjectPtr<UTakeRecorderSource>, TObjectPtr<ULevelSequence>>& SourceSubSequence : SourceSubSequenceMap)
 	{
 		// Only save subsequences but not the root sequence, it is already saved in UTakeRecorder::StopInternal
 		const bool bIsValidSubSequence = SourceSubSequence.Value && SourceSubSequence.Value.Get() != InSequence;
@@ -641,7 +664,8 @@ void UTakeRecorderSources::StopRecording(class ULevelSequence* InSequence, const
 			}
 			else if (Settings.bSaveRecordedAssets)
 			{
-				TakesUtils::SaveAsset(SourceSubSequence.Value, Settings.bRecordToPossessable);
+				const bool bIsPossessable = TakesUtils::IsSourcePossessable(SourceSubSequence.Key);
+				TakesUtils::SaveAsset(SourceSubSequence.Value, bIsPossessable);
 			}
 		}
 	}
diff --git a/Engine/Plugins/VirtualProduction/Takes/Source/TakesCore/Private/TakesUtils.cpp b/Engine/Plugins/VirtualProduction/Takes/Source/TakesCore/Private/TakesUtils.cpp
--- a/Engine/Plugins/VirtualProduction/Takes/Source/TakesCore/Private/TakesUtils.cpp
+++ b/Engine/Plugins/VirtualProduction/Takes/Source/TakesCore/Private/TakesUtils.cpp
@@ -6,18 +6,20 @@
 
 #include "Engine/Engine.h"
 #include "Engine/World.h"
-#include "LevelSequence.h"
 #include "ILevelSequenceEditorToolkit.h"
-#include "MovieScene.h"
+#include "LevelSequence.h"
 #include "Math/Range.h"
+#include "MovieScene.h"
 #include "MovieSceneSection.h"
 #include "MovieSceneTimeHelpers.h"
-#include "HAL/FileManager.h"
-#include "Tracks/MovieSceneCameraCutTrack.h"
 #include "Sections/MovieSceneCameraCutSection.h"
+#include "TakeRecorderSource.h"
+#include "TrackRecorders/IMovieSceneTrackRecorderHost.h"
+#include "Tracks/MovieSceneCameraCutTrack.h"
 #include "UObject/SavePackage.h"
 
 #include "Serialization/ArchiveReplaceOrClearExternalReferences.h"
+
 #define LOCTEXT_NAMESPACE "TakesUtils"
 
 namespace TakesUtils
@@ -47,8 +49,11 @@ public:
 	virtual FArchive& operator<<(UObject*& Obj) override
 	{
 		UObject* Original = Obj;
-		Super::operator<<(Obj);
-
+		// We don't want to clear references if we're dealing with a spawnable.
+		if (!IsObjectPossessable(Obj))
+		{
+			Super::operator<<(Obj);
+		}
 		if (Original && !Obj)
 		{
 			RemovedObjects.AddUnique(Original);
@@ -56,6 +61,23 @@ public:
 		
 		return *this;
 	}
+	
+	/** Checks if the given object is considered possessable. */
+	bool IsObjectPossessable(UObject* Obj) const
+	{
+		if (const ULevelSequence* LevelSequence = Cast<ULevelSequence>(SearchObject))
+		{
+			if (const AActor* Actor = Cast<AActor>(Obj))
+			{
+				if (const UTakeRecorderSources* Sources = LevelSequence->FindMetaData<UTakeRecorderSources>())
+				{
+					return Sources->IsActorPossessable(Actor);
+				}
+			}
+		}
+		
+		return false;
+	}
 };
 
 static void LogRemovedExternalObjects(UObject* ReferencingObject, const TArray<UObject*>& InRemovedObjects)
@@ -160,7 +182,7 @@ void CleanupExternalReferences(UObject* InObject)
 	CleanupExternalReferences(InObject, Package);
 }
 	
-void SaveAsset(UObject* InObject, bool bIsPossessable)
+void SaveAsset(UObject* InObject, bool bAlwaysKeepExternalReferences)
 {
 	if (!InObject)
 	{
@@ -180,7 +202,7 @@ void SaveAsset(UObject* InObject, bool bIsPossessable)
 
 	double StartTime = FPlatformTime::Seconds();
 
-	if (!bIsPossessable)
+	if (!bAlwaysKeepExternalReferences)
 	{
 		CleanupExternalReferences(InObject, Package);		
 	}
@@ -346,6 +368,19 @@ FGuid ResolveActorFromSequence(const AActor* InActor, const ULevelSequence* Curr
 	// to the sequence yet.
 	return FGuid();
 }
+
+bool IsSourcePossessable(const UTakeRecorderSource* InSource)
+{
+	check(InSource);
+	if (const IMovieSceneTrackRecorderHost* RecorderHost = InSource->AsTrackRecorderHost())
+	{
+		return RecorderHost->GetTrackRecorderSettings().bRecordToPossessable;
+	}
+	
+	const UTakeRecorderSources* Sources = InSource->GetTypedOuter<UTakeRecorderSources>();
+	check(Sources);
+	return Sources->GetSettings().bRecordToPossessable;
+}
 }
 
 #undef LOCTEXT_NAMESPACE
diff --git a/Engine/Plugins/VirtualProduction/Takes/Source/TakesCore/Public/TakeRecorderSource.h b/Engine/Plugins/VirtualProduction/Takes/Source/TakesCore/Public/TakeRecorderSource.h
--- a/Engine/Plugins/VirtualProduction/Takes/Source/TakesCore/Public/TakeRecorderSource.h
+++ b/Engine/Plugins/VirtualProduction/Takes/Source/TakesCore/Public/TakeRecorderSource.h
@@ -204,7 +204,7 @@ public:
 		FText Result = GetAddSourceDisplayTextImpl();
 		if (Result.IsEmpty())
 		{
-			ensureMsgf(false, TEXT("You must provide localized source name."));
+			ensureMsgf(IsTemplate(RF_ClassDefaultObject) /* This can be called on the CDO. */, TEXT("You must provide localized source name."));
 			return GetDisplayText();
 		}
 		return Result;
@@ -230,6 +230,12 @@ public:
 	 * Whether or the source can be referenced via take number
 	 */
 	virtual bool SupportsTakeNumber() const { return true; }
+	
+	/** If this source implements the IMovieSceneTrackRecorderHost interface. This is used to avoid RTTI. */
+	virtual const class IMovieSceneTrackRecorderHost* AsTrackRecorderHost() const { return nullptr; }
+	
+	/** If this source is an actor source for the given actor. Used for cross-module communication. */
+	virtual bool IsActorSourceForActor(const AActor* Actor) const { return false; }
 
 private:
 
diff --git a/Engine/Plugins/VirtualProduction/Takes/Source/TakesCore/Public/TakeRecorderSources.h b/Engine/Plugins/VirtualProduction/Takes/Source/TakesCore/Public/TakeRecorderSources.h
--- a/Engine/Plugins/VirtualProduction/Takes/Source/TakesCore/Public/TakeRecorderSources.h
+++ b/Engine/Plugins/VirtualProduction/Takes/Source/TakesCore/Public/TakeRecorderSources.h
@@ -163,6 +163,9 @@ public:
 	 */
 	UE_API ULevelSequence* AddSourceForActorBinding(AActor* InActor, ULevelSequence* InLevelSequence);
 	
+	/** Checks if a given actor is possessable given source configuration. */
+	UE_API bool IsActorPossessable(const AActor* InActor) const;
+	
 public:
 
 	/*
@@ -269,6 +272,9 @@ private:
 	/** List of sub-sections that we're recording into. Needed to ensure they're all the right size at the end without re-adjusting every sub-section in a sequence. */
 	UPROPERTY(Transient)
 	TArray<TObjectPtr<class UMovieSceneSubSection>> ActiveSubSections;
+	
+	/** Bindings added for an actor mapped to whether it is possessable or not. Reset on PreRecord. */
+	TMap<TWeakObjectPtr<AActor>, bool> AddedActorPossessableBindings;
 
 	/** Are we currently in a recording pass and should be ticking our Sources? */
 	bool bIsRecording;
diff --git a/Engine/Plugins/VirtualProduction/Takes/Source/TakesCore/Public/TakesUtils.h b/Engine/Plugins/VirtualProduction/Takes/Source/TakesCore/Public/TakesUtils.h
--- a/Engine/Plugins/VirtualProduction/Takes/Source/TakesCore/Public/TakesUtils.h
+++ b/Engine/Plugins/VirtualProduction/Takes/Source/TakesCore/Public/TakesUtils.h
@@ -15,6 +15,7 @@
 #include "Misc/Paths.h"
 #include "PackageTools.h"
 
+class UTakeRecorderSource;
 class UWorld;
 class ULevelSequence;
 
@@ -32,7 +33,13 @@ namespace TakesUtils
 
 	TAKESCORE_API void CleanupExternalReferences(UObject* InObject);
 	
-	TAKESCORE_API void SaveAsset(UObject* InObject, bool bIsPossessable = true);
+	/** 
+	 * Save an object to disk. 
+	 * @param InObject The object to save.
+	 * @param bAlwaysKeepExternalReferences If we should always keep external references.
+	 * When false, we clear them. When false and the InObject is a Level Sequence, we perform fine-grained cleanup.
+	 */
+	TAKESCORE_API void SaveAsset(UObject* InObject, bool bAlwaysKeepExternalReferences = true);
 
 	TAKESCORE_API void CreateCameraCutTrack(ULevelSequence* LevelSequence, const FGuid& RecordedCameraGuid, const FMovieSceneSequenceID& SequenceID, const TRange<FFrameNumber>& InRange);
 
@@ -178,4 +185,7 @@ namespace TakesUtils
 	
 	/** Retrieve an actor guid for moviescene. */
 	TAKESCORE_API FGuid ResolveActorFromSequence(const AActor* InActor, const ULevelSequence* CurrentSequence);
+	
+	/** Checks if a source is considered possessable. */
+	TAKESCORE_API bool IsSourcePossessable(const UTakeRecorderSource* InSource);
 }
