TakeRecorder: Using an existing recording as a base will correctly open the newly set pending take and sequencer again.

#Synced-CL 51303076

diff --git a/Engine/Plugins/VirtualProduction/Takes/Source/TakeRecorder/Private/Widgets/STakeRecorderCockpit.cpp b/Engine/Plugins/VirtualProduction/Takes/Source/TakeRecorder/Private/Widgets/STakeRecorderCockpit.cpp
--- a/Engine/Plugins/VirtualProduction/Takes/Source/TakeRecorder/Private/Widgets/STakeRecorderCockpit.cpp
+++ b/Engine/Plugins/VirtualProduction/Takes/Source/TakeRecorder/Private/Widgets/STakeRecorderCockpit.cpp
@@ -76,6 +76,7 @@
 #include "NamingTokens.h"
 #include "NamingTokensEngineSubsystem.h"
 #include "NamingTokensSpecifiers.h"
+#include "STakeRecorderPanel.h"
 #include "TakeRecorderNamingTokensData.h"
 #include "Timecode/HitchProtectionFrameRateModel.h"
 #include "Widgets/Layout/SScaleBox.h"
@@ -1069,6 +1070,11 @@ FReply STakeRecorderCockpit::NewRecordingFromThis()
 	{
 		TSharedRef<STakeRecorderTabContent> TabContent = StaticCastSharedRef<STakeRecorderTabContent>(DockTab->GetContent());
 		TabContent->SetupForRecording(Sequence);
+		if (const TSharedPtr<STakeRecorderPanel> TakeRecorderPanel = TabContent->GetTakeRecorderPanel().Pin())
+		{
+			// Go through panel rather than subsystem so sequencer remains open if necessary.
+			TakeRecorderPanel->ResetToPendingTake();
+		}
 	}
 
 	return FReply::Handled();
@@ -1692,4 +1698,4 @@ void STakeRecorderCockpit::CreateUserTokensUI()
 	}
 }
 
-#undef LOCTEXT_NAMESPACE
+#undef LOCTEXT_NAMESPACE
\ No newline at end of file
diff --git a/Engine/Plugins/VirtualProduction/Takes/Source/TakeRecorder/Private/Widgets/STakeRecorderPanel.cpp b/Engine/Plugins/VirtualProduction/Takes/Source/TakeRecorder/Private/Widgets/STakeRecorderPanel.cpp
--- a/Engine/Plugins/VirtualProduction/Takes/Source/TakeRecorder/Private/Widgets/STakeRecorderPanel.cpp
+++ b/Engine/Plugins/VirtualProduction/Takes/Source/TakeRecorder/Private/Widgets/STakeRecorderPanel.cpp
@@ -531,6 +531,12 @@ void STakeRecorderPanel::ClearPendingTake()
 	TakeRecorderSubsystem->ClearPendingTake();
 }
 
+void STakeRecorderPanel::ResetToPendingTake()
+{
+	TakeRecorderSubsystem->ResetToPendingTake();
+	RefreshPanel();
+}
+
 TOptional<ETakeRecorderPanelMode> STakeRecorderPanel::GetMode() const
 {
 	if (TakeRecorderSubsystem->GetSuppliedLevelSequence())
@@ -726,9 +732,7 @@ void STakeRecorderPanel::OnSaveAsPreset()
 
 FReply STakeRecorderPanel::OnBackToPendingTake()
 {
-	TakeRecorderSubsystem->ResetToPendingTake();
-	RefreshPanel();
-
+	ResetToPendingTake();
 	return FReply::Handled();
 }
 
diff --git a/Engine/Plugins/VirtualProduction/Takes/Source/TakeRecorder/Private/Widgets/STakeRecorderPanel.h b/Engine/Plugins/VirtualProduction/Takes/Source/TakeRecorder/Private/Widgets/STakeRecorderPanel.h
--- a/Engine/Plugins/VirtualProduction/Takes/Source/TakeRecorder/Private/Widgets/STakeRecorderPanel.h
+++ b/Engine/Plugins/VirtualProduction/Takes/Source/TakeRecorder/Private/Widgets/STakeRecorderPanel.h
@@ -74,6 +74,7 @@ public:
 	TSharedPtr<STakeRecorderCockpit> GetCockpitWidget() const { return CockpitWidget; }
 
 	void ClearPendingTake();
+	void ResetToPendingTake();
 
 	TOptional<ETakeRecorderPanelMode> GetMode() const;
 
diff --git a/Engine/Plugins/VirtualProduction/Takes/Source/TakeRecorder/Private/Widgets/STakeRecorderTabContent.h b/Engine/Plugins/VirtualProduction/Takes/Source/TakeRecorder/Private/Widgets/STakeRecorderTabContent.h
--- a/Engine/Plugins/VirtualProduction/Takes/Source/TakeRecorder/Private/Widgets/STakeRecorderTabContent.h
+++ b/Engine/Plugins/VirtualProduction/Takes/Source/TakeRecorder/Private/Widgets/STakeRecorderTabContent.h
@@ -64,6 +64,8 @@ public:
 	void ClearPendingTake();
 
 	bool CanStartRecording(FText& ErrorText) const;
+	
+	const TWeakPtr<STakeRecorderPanel>& GetTakeRecorderPanel() const { return WeakPanel; }
 
 private:
 
