TakeRecorder: Increment take number when using an existing recording as a base.
Add IncrementTakeNumber to subsystem API.

#Synced-CL 51433690

diff --git a/Engine/Plugins/VirtualProduction/Takes/Source/TakeRecorder/Private/Recorder/TakeRecorderSubsystem.cpp b/Engine/Plugins/VirtualProduction/Takes/Source/TakeRecorder/Private/Recorder/TakeRecorderSubsystem.cpp
--- a/Engine/Plugins/VirtualProduction/Takes/Source/TakeRecorder/Private/Recorder/TakeRecorderSubsystem.cpp
+++ b/Engine/Plugins/VirtualProduction/Takes/Source/TakeRecorder/Private/Recorder/TakeRecorderSubsystem.cpp
@@ -155,6 +155,11 @@ void UTakeRecorderSubsystem::GetNumberOfTakes(const FString& InSlate, int32& Out
 	Implementation->GetNumberOfTakes(InSlate, OutMaxTake, OutNumTakes);
 }
 
+void UTakeRecorderSubsystem::IncrementTakeNumber()
+{
+	Implementation->IncrementTakeNumber();
+}
+
 TArray<FAssetData> UTakeRecorderSubsystem::GetSlates(FName InPackagePath) const
 {
 	return Implementation->GetSlates(InPackagePath);
diff --git a/Engine/Plugins/VirtualProduction/Takes/Source/TakeRecorder/Private/Recorder/TakeRecorderSubsystemImplementation.h b/Engine/Plugins/VirtualProduction/Takes/Source/TakeRecorder/Private/Recorder/TakeRecorderSubsystemImplementation.h
--- a/Engine/Plugins/VirtualProduction/Takes/Source/TakeRecorder/Private/Recorder/TakeRecorderSubsystemImplementation.h
+++ b/Engine/Plugins/VirtualProduction/Takes/Source/TakeRecorder/Private/Recorder/TakeRecorderSubsystemImplementation.h
@@ -62,6 +62,7 @@ public:
 	virtual void SetTakeNumber(int32 InNewTakeNumber, bool bEmitChanged = true) override;
 	virtual int32 GetNextTakeNumber(const FString& InSlate) const override;
 	virtual void GetNumberOfTakes(const FString& InSlate, int32& OutMaxTake, int32& OutNumTakes) const override;
+	virtual void IncrementTakeNumber() override;
 	virtual TArray<FAssetData> GetSlates(FName InPackagePath = NAME_None) const override;
 	virtual void SetSlateName(const FString& InSlateName, bool bEmitChanged = true) override;
 	virtual bool MarkFrame() override;
@@ -105,9 +106,6 @@ private:
 	/** Sets the frame rate. */
 	void SetFrameRateImpl(const FFrameRate& InFrameRate, bool bFromTimecode);
 
-	/** Calculate and apply the next take number. */
-	void IncrementTakeNumber();
-
 	void OnAssetRegistryFilesLoaded();
 	void OnRecordingInitialized(UTakeRecorder* Recorder);
 	void OnRecordingStarted(UTakeRecorder* Recorder);
diff --git a/Engine/Plugins/VirtualProduction/Takes/Source/TakeRecorder/Private/Widgets/STakeRecorderCockpit.cpp b/Engine/Plugins/VirtualProduction/Takes/Source/TakeRecorder/Private/Widgets/STakeRecorderCockpit.cpp
--- a/Engine/Plugins/VirtualProduction/Takes/Source/TakeRecorder/Private/Widgets/STakeRecorderCockpit.cpp
+++ b/Engine/Plugins/VirtualProduction/Takes/Source/TakeRecorder/Private/Widgets/STakeRecorderCockpit.cpp
@@ -1074,6 +1074,7 @@ FReply STakeRecorderCockpit::NewRecordingFromThis()
 		{
 			// Go through panel rather than subsystem so sequencer remains open if necessary.
 			TakeRecorderPanel->ResetToPendingTake();
+			TakeRecorderSubsystem->IncrementTakeNumber();
 		}
 	}
 
diff --git a/Engine/Plugins/VirtualProduction/Takes/Source/TakeRecorder/Public/Recorder/ITakeRecorderSubsystemInterface.h b/Engine/Plugins/VirtualProduction/Takes/Source/TakeRecorder/Public/Recorder/ITakeRecorderSubsystemInterface.h
--- a/Engine/Plugins/VirtualProduction/Takes/Source/TakeRecorder/Public/Recorder/ITakeRecorderSubsystemInterface.h
+++ b/Engine/Plugins/VirtualProduction/Takes/Source/TakeRecorder/Public/Recorder/ITakeRecorderSubsystemInterface.h
@@ -55,6 +55,7 @@ public:
 	virtual void SetTakeNumber(int32 InNewTakeNumber, bool bEmitChanged = true) = 0;
 	virtual int32 GetNextTakeNumber(const FString& InSlate) const = 0;
 	virtual void GetNumberOfTakes(const FString& InSlate, int32& OutMaxTake, int32& OutNumTakes) const = 0;
+	virtual void IncrementTakeNumber() = 0;
 	virtual TArray<FAssetData> GetSlates(FName InPackagePath = NAME_None) const = 0;
 	virtual void SetSlateName(const FString& InSlateName, bool bEmitChanged = true) = 0;
 	virtual bool MarkFrame() = 0;
diff --git a/Engine/Plugins/VirtualProduction/Takes/Source/TakeRecorder/Public/Recorder/TakeRecorderSubsystem.h b/Engine/Plugins/VirtualProduction/Takes/Source/TakeRecorder/Public/Recorder/TakeRecorderSubsystem.h
--- a/Engine/Plugins/VirtualProduction/Takes/Source/TakeRecorder/Public/Recorder/TakeRecorderSubsystem.h
+++ b/Engine/Plugins/VirtualProduction/Takes/Source/TakeRecorder/Public/Recorder/TakeRecorderSubsystem.h
@@ -189,6 +189,10 @@ public:
 	UFUNCTION(BlueprintCallable, Category = TakeRecorder)
 	UE_API virtual void GetNumberOfTakes(const FString& InSlate, int32& OutMaxTake, int32& OutNumTakes) const override;
 
+	/** Calculate and apply the next take number. */
+	UFUNCTION(BlueprintCallable, Category = TakeRecorder)
+	UE_API virtual void IncrementTakeNumber() override;
+	
 	/**
 	 * Retrieve all slates.
 	 * @param InPackagePath [Optional] Scope the search to a specific folder, recursively.
