TakeRecorder: Using an existing recording as a base with spawnable sources will no longer overwrite the original source actor in the new recording.
- Add project and per source setting to allow legacy overwrite behavior instead.
- Remove public interface API recently added which went unused.

#Synced-CL 51673377

diff --git a/Engine/Plugins/VirtualProduction/Takes/Source/TakeRecorder/Private/Recorder/TakeRecorder.cpp b/Engine/Plugins/VirtualProduction/Takes/Source/TakeRecorder/Private/Recorder/TakeRecorder.cpp
--- a/Engine/Plugins/VirtualProduction/Takes/Source/TakeRecorder/Private/Recorder/TakeRecorder.cpp
+++ b/Engine/Plugins/VirtualProduction/Takes/Source/TakeRecorder/Private/Recorder/TakeRecorder.cpp
@@ -930,6 +930,7 @@ void UTakeRecorder::PreRecord()
 	TakeRecorderSourcesSettings.bStartAtCurrentTimecode = Parameters.Project.bStartAtCurrentTimecode;
 	TakeRecorderSourcesSettings.bRecordSourcesIntoSubSequences = Parameters.Project.bRecordSourcesIntoSubSequences;
 	TakeRecorderSourcesSettings.bRecordToPossessable = Parameters.Project.bRecordToPossessable;
+	TakeRecorderSourcesSettings.SpawnableOverwriteBehavior = Parameters.Project.SpawnableOverwriteBehavior;
 	TakeRecorderSourcesSettings.bSaveRecordedAssets = (!Parameters.bDisableRecordingAndSave && Parameters.User.bSaveRecordedAssets) || GEditor == nullptr;
 	TakeRecorderSourcesSettings.bRemoveRedundantTracks = Parameters.User.bRemoveRedundantTracks;
 	TakeRecorderSourcesSettings.bAutoLock = Parameters.User.bAutoLock;
diff --git a/Engine/Plugins/VirtualProduction/Takes/Source/TakeRecorder/Private/Recorder/TakeRecorderParameters.cpp b/Engine/Plugins/VirtualProduction/Takes/Source/TakeRecorder/Private/Recorder/TakeRecorderParameters.cpp
--- a/Engine/Plugins/VirtualProduction/Takes/Source/TakeRecorder/Private/Recorder/TakeRecorderParameters.cpp
+++ b/Engine/Plugins/VirtualProduction/Takes/Source/TakeRecorder/Private/Recorder/TakeRecorderParameters.cpp
@@ -31,8 +31,10 @@ FTakeRecorderProjectParameters::FTakeRecorderProjectParameters()
 	, bRecordTimecode(false)
 	, bRecordSourcesIntoSubSequences(false)
 	, bRecordToPossessable(false)
+	, SpawnableOverwriteBehavior(ETakeRecorderSpawnableOverwriteBehavior::CreateNew)
 	, bShowNotifications(true)
-{}
+{
+}
 
 FString FTakeRecorderProjectParameters::GetResolvedRootTakeSaveDir() const
 {
diff --git a/Engine/Plugins/VirtualProduction/Takes/Source/TakeRecorder/Private/Recorder/TakeRecorderSubsystem.cpp b/Engine/Plugins/VirtualProduction/Takes/Source/TakeRecorder/Private/Recorder/TakeRecorderSubsystem.cpp
--- a/Engine/Plugins/VirtualProduction/Takes/Source/TakeRecorder/Private/Recorder/TakeRecorderSubsystem.cpp
+++ b/Engine/Plugins/VirtualProduction/Takes/Source/TakeRecorder/Private/Recorder/TakeRecorderSubsystem.cpp
@@ -85,6 +85,11 @@ void UTakeRecorderSubsystem::RevertChanges()
 	Implementation->RevertChanges();
 }
 
+bool UTakeRecorderSubsystem::IsRecordingUsingExistingSequence() const
+{
+	return Implementation->IsRecordingUsingExistingSequence();
+}
+
 UTakeRecorderSource* UTakeRecorderSubsystem::AddSource(const TSubclassOf<UTakeRecorderSource> InSourceClass)
 {
 	return Implementation->AddSource(InSourceClass);
diff --git a/Engine/Plugins/VirtualProduction/Takes/Source/TakeRecorder/Private/Recorder/TakeRecorderSubsystemImplementation.cpp b/Engine/Plugins/VirtualProduction/Takes/Source/TakeRecorder/Private/Recorder/TakeRecorderSubsystemImplementation.cpp
--- a/Engine/Plugins/VirtualProduction/Takes/Source/TakeRecorder/Private/Recorder/TakeRecorderSubsystemImplementation.cpp
+++ b/Engine/Plugins/VirtualProduction/Takes/Source/TakeRecorder/Private/Recorder/TakeRecorderSubsystemImplementation.cpp
@@ -156,6 +156,8 @@ void UTakeRecorderSubsystemImplementation::SetTargetSequence(const FTakeRecorder
 	
 	TransientPreset = AllocateTransientPreset();
 	
+	bRecordingUsingExistingSequence = false;
+	
 	// Copy the base preset into the transient preset if it was provided.
 	// We do this first so that anything that asks for its Level Sequence
 	// on construction gets the right one
@@ -165,6 +167,7 @@ void UTakeRecorderSubsystemImplementation::SetTargetSequence(const FTakeRecorder
 	}
 	else if (InData.BaseSequence)
 	{
+		bRecordingUsingExistingSequence = true;
 		TransientPreset->CopyFrom(InData.BaseSequence);
 
 		ULevelSequence* LevelSequence = TransientPreset->GetLevelSequence();
@@ -188,6 +191,7 @@ void UTakeRecorderSubsystemImplementation::SetTargetSequence(const FTakeRecorder
 	}
 	else if (InData.RecordIntoSequence)
 	{
+		bRecordingUsingExistingSequence = true;
 		SetRecordIntoLevelSequence(InData.RecordIntoSequence);
 	}
 	else if (InData.SequenceToView)
@@ -382,7 +386,10 @@ void UTakeRecorderSubsystemImplementation::ResetToPendingTake()
 	{
 		LastRecordedLevelSequence = SuppliedLevelSequence;
 	}
-
+	
+	// We purposely don't reset bRecordingUsingExistingSequence here, because the pending take may still be considered recorded over and
+	// certain workflows from the panel rely on this. ClearPendingTake may be called if we want to reset bRecordingUsingExistingSequence.
+	
 	SuppliedLevelSequence = nullptr;
 	RecordIntoLevelSequence = nullptr;
 	
@@ -391,6 +398,8 @@ void UTakeRecorderSubsystemImplementation::ResetToPendingTake()
 
 void UTakeRecorderSubsystemImplementation::ClearPendingTake()
 {
+	bRecordingUsingExistingSequence = false;
+	
 	if (IsReviewing())
 	{
 		LastRecordedLevelSequence = SuppliedLevelSequence;
@@ -435,6 +444,11 @@ void UTakeRecorderSubsystemImplementation::RevertChanges()
 	TransientPreset->GetOutermost()->SetDirtyFlag(false);
 }
 
+bool UTakeRecorderSubsystemImplementation::IsRecordingUsingExistingSequence() const
+{
+	return GetTakeRecorderMode() == ETakeRecorderMode::RecordIntoSequence || bRecordingUsingExistingSequence;
+}
+
 UTakeRecorderSource* UTakeRecorderSubsystemImplementation::AddSource(const TSubclassOf<UTakeRecorderSource> InSourceClass)
 {
 	UTakeRecorderSources* Sources = GetSources();
diff --git a/Engine/Plugins/VirtualProduction/Takes/Source/TakeRecorder/Private/Recorder/TakeRecorderSubsystemImplementation.h b/Engine/Plugins/VirtualProduction/Takes/Source/TakeRecorder/Private/Recorder/TakeRecorderSubsystemImplementation.h
--- a/Engine/Plugins/VirtualProduction/Takes/Source/TakeRecorder/Private/Recorder/TakeRecorderSubsystemImplementation.h
+++ b/Engine/Plugins/VirtualProduction/Takes/Source/TakeRecorder/Private/Recorder/TakeRecorderSubsystemImplementation.h
@@ -48,6 +48,7 @@ public:
 	virtual void ClearPendingTake() override;
 	virtual UTakePreset* GetPendingTake() const override;
 	virtual void RevertChanges() override;
+	virtual bool IsRecordingUsingExistingSequence() const override;
 	virtual UTakeRecorderSource* AddSource(const TSubclassOf<UTakeRecorderSource> InSourceClass) override;
 	virtual void RemoveSource(UTakeRecorderSource* InSource) override;
 	virtual void ClearSources() override;
@@ -190,4 +191,7 @@ private:
 
 	/** If SetTargetSequence has been called, and we're fully initialized. */
 	bool bHasTargetSequenceBeenSet = false;
+	
+	/** Helps determine if we are recording into an existing take or using one as a base. */
+	bool bRecordingUsingExistingSequence = false;
 };
diff --git a/Engine/Plugins/VirtualProduction/Takes/Source/TakeRecorder/Public/Recorder/ITakeRecorderSubsystemInterface.h b/Engine/Plugins/VirtualProduction/Takes/Source/TakeRecorder/Public/Recorder/ITakeRecorderSubsystemInterface.h
--- a/Engine/Plugins/VirtualProduction/Takes/Source/TakeRecorder/Public/Recorder/ITakeRecorderSubsystemInterface.h
+++ b/Engine/Plugins/VirtualProduction/Takes/Source/TakeRecorder/Public/Recorder/ITakeRecorderSubsystemInterface.h
@@ -37,6 +37,7 @@ public:
 	virtual void ClearPendingTake() = 0;
 	virtual UTakePreset* GetPendingTake() const = 0;
 	virtual void RevertChanges() = 0;
+	virtual bool IsRecordingUsingExistingSequence() const = 0;
 	
 	// ~Begin Sources
 	virtual UTakeRecorderSource* AddSource(const TSubclassOf<UTakeRecorderSource> InSourceClass) = 0;
diff --git a/Engine/Plugins/VirtualProduction/Takes/Source/TakeRecorder/Public/Recorder/TakeRecorderParameters.h b/Engine/Plugins/VirtualProduction/Takes/Source/TakeRecorder/Public/Recorder/TakeRecorderParameters.h
--- a/Engine/Plugins/VirtualProduction/Takes/Source/TakeRecorder/Public/Recorder/TakeRecorderParameters.h
+++ b/Engine/Plugins/VirtualProduction/Takes/Source/TakeRecorder/Public/Recorder/TakeRecorderParameters.h
@@ -5,6 +5,7 @@
 #include "UObject/Object.h"
 #include "Engine/EngineTypes.h"
 #include "MovieSceneFwd.h"
+#include "TakeRecorderSpawnableBehavior.h"
 #include "TrackRecorders/IMovieSceneTrackRecorderHost.h"
 #include "TakeRecorderParameters.generated.h"
 
@@ -144,6 +145,10 @@ struct FTakeRecorderProjectParameters
 	UPROPERTY(config, EditAnywhere, BlueprintReadWrite, Category = "Take Recorder")
 	bool bRecordToPossessable;
 
+	/** If spawnables should be overwritten during record. */
+	UPROPERTY(config, EditAnywhere, BlueprintReadWrite, Category = "Take Recorder", meta = (ValidEnumValues = "OverwriteLegacy, CreateNew"))
+	ETakeRecorderSpawnableOverwriteBehavior SpawnableOverwriteBehavior;
+	
 	/** List of property names for which movie scene tracks will always record. */
 	UPROPERTY(config, EditAnywhere, Category = "Take Recorder")
 	TArray<FTakeRecorderTrackSettings> DefaultTracks;
diff --git a/Engine/Plugins/VirtualProduction/Takes/Source/TakeRecorder/Public/Recorder/TakeRecorderSubsystem.h b/Engine/Plugins/VirtualProduction/Takes/Source/TakeRecorder/Public/Recorder/TakeRecorderSubsystem.h
--- a/Engine/Plugins/VirtualProduction/Takes/Source/TakeRecorder/Public/Recorder/TakeRecorderSubsystem.h
+++ b/Engine/Plugins/VirtualProduction/Takes/Source/TakeRecorder/Public/Recorder/TakeRecorderSubsystem.h
@@ -116,6 +116,10 @@ public:
 	UFUNCTION(BlueprintCallable, Category = TakeRecorder)
 	UE_API virtual void RevertChanges() override;
 	
+	/** If we are using an existing sequence as a base or recording into an existing one. */
+	UFUNCTION(BlueprintCallable, Category = TakeRecorder)
+	UE_API virtual bool IsRecordingUsingExistingSequence() const override;
+	
 	// ~Begin Sources
 	/** Add a source by a source class. */
 	UFUNCTION(BlueprintCallable, Category = TakeRecorder, meta = (DeterminesOutputType = "InSourceClass"))
diff --git a/Engine/Plugins/VirtualProduction/Takes/Source/TakeRecorderSources/Private/TakeRecorderActorSource.cpp b/Engine/Plugins/VirtualProduction/Takes/Source/TakeRecorderSources/Private/TakeRecorderActorSource.cpp
--- a/Engine/Plugins/VirtualProduction/Takes/Source/TakeRecorderSources/Private/TakeRecorderActorSource.cpp
+++ b/Engine/Plugins/VirtualProduction/Takes/Source/TakeRecorderSources/Private/TakeRecorderActorSource.cpp
@@ -49,6 +49,7 @@
 #include "Tracks/MovieSceneSpawnTrack.h"
 #endif
 #include "NamingTokens/TakeRecorderNamingTokensContext.h"
+#include "Recorder/TakeRecorderSubsystem.h"
 #include "SequencerUtilities.h"
 #include "TakeRecorderSourcesModule.h"
 #include "TakeRecorderTimeProcessing.h"
@@ -247,10 +248,16 @@ TArray<UTakeRecorderSource*> UTakeRecorderActorSource::PreRecording(ULevelSequen
 	UMovieScene* MovieScene = TargetLevelSequence->GetMovieScene();
 
 	// Look to see if the movie scene already has our object binding in it (which is common if we're recording a new take) so we can
-	// replace the data that was already there.
+	// replace the data that was already there...
 	CachedObjectBindingGuid = ResolveActorFromSequence(ActorToRecord, TargetLevelSequence);
-	CleanExistingDataFromSequence(CachedObjectBindingGuid, *TargetLevelSequence);
-		
+	
+	// ... but don't clear existing data if we're recording into an existing take as a spawnable. We want the original object, and its recorded data,
+	// to persist throughout this new recording.
+	if (OverwriteExistingTrackData())
+	{
+		CleanExistingDataFromSequence(CachedObjectBindingGuid, *TargetLevelSequence);
+	}
+	
 	// Initialize the header for this actor in the Manifest Serializer for streaming data capture.
 	FName SerializedType("Actor");
 	FActorFileHeader Header(ObjectBindingName, ActorToRecord->GetActorLabel(), SerializedType, ActorToRecord->GetClass()->GetPathName(), false);
@@ -1849,6 +1856,39 @@ bool UTakeRecorderActorSource::IsActorSourceForActor(const AActor* Actor) const
 #endif
 }
 
+bool UTakeRecorderActorSource::OverwriteExistingTrackData() const
+{
+	const UTakeRecorderSubsystem* TakeRecorderSubsystem = GEngine->GetEngineSubsystem<UTakeRecorderSubsystem>();
+	check(TakeRecorderSubsystem);
+	if (!TakeRecorderSubsystem->IsRecordingUsingExistingSequence() || GetRecordToPossessable())
+	{
+		// If we're creating a completely new recording or using possessables we always want to overwrite.
+		
+		return true;
+	}
+	
+	switch (SpawnableOverwriteBehavior)
+	{
+	case ETakeRecorderSpawnableOverwriteBehavior::OverwriteLegacy:
+		return true;
+
+	case ETakeRecorderSpawnableOverwriteBehavior::CreateNew:
+		return false;
+
+	case ETakeRecorderSpawnableOverwriteBehavior::ProjectDefault:
+		if (const UTakeRecorderSources* SourcesList = GetTypedOuter<UTakeRecorderSources>())
+		{
+			return SourcesList->GetSettings().SpawnableOverwriteBehavior == ETakeRecorderSpawnableOverwriteBehavior::OverwriteLegacy;
+		}
+		break;
+
+	default:
+		break;
+	}
+
+	return Super::OverwriteExistingTrackData();
+}
+
 void UTakeRecorderActorSource::SetSourceActor(TSoftObjectPtr<AActor> InTarget)
 {
 	Target = InTarget;
@@ -1888,11 +1928,6 @@ void UTakeRecorderActorSource::SetIsTemporarySource(bool bValue)
 	bIsTemporarySource = bValue;
 }
 
-bool UTakeRecorderActorSource::IsOtherActorPossessable(AActor* InActor) const
-{
-	return TakeRecorderSourcesUtils::IsActorPossessable(this, InActor);
-}
-
 void UTakeRecorderActorSource::InitializeFactorySettingsObject(TSubclassOf<UMovieSceneTrackRecorderSettings> InClass)
 {
 	ensure(InClass);
diff --git a/Engine/Plugins/VirtualProduction/Takes/Source/TakeRecorderSources/Private/TakeRecorderSourcesUtils.h b/Engine/Plugins/VirtualProduction/Takes/Source/TakeRecorderSources/Private/TakeRecorderSourcesUtils.h
--- a/Engine/Plugins/VirtualProduction/Takes/Source/TakeRecorderSources/Private/TakeRecorderSourcesUtils.h
+++ b/Engine/Plugins/VirtualProduction/Takes/Source/TakeRecorderSources/Private/TakeRecorderSourcesUtils.h
@@ -92,38 +92,6 @@ namespace TakeRecorderSourcesUtils
 
 		return false;
 	}
-	
-	/** If an actor is configured to be possessable. */
-	static bool IsActorPossessable(const UTakeRecorderActorSource* InSource, AActor* Actor)
-	{
-		check(InSource);
-		
-		if (InSource->Target == Actor && InSource->GetRecordToPossessable())
-		{
-			return true;
-		}
-
-		const UTakeRecorderSources* OwningSources = CastChecked<UTakeRecorderSources>(InSource->GetOuter());
-		for (UTakeRecorderSource* Source : OwningSources->GetSources())
-		{
-			if (const UTakeRecorderActorSource* ActorSource = Cast<UTakeRecorderActorSource>(Source))
-			{
-				if (ActorSource->bEnabled && ActorSource->Target.Get() == Actor && ActorSource->GetRecordToPossessable())
-				{
-					return true;
-				}
-			}
-		}
-
-		// Fallback to global settings in the event the actor doesn't have a source for it. This is what will be used
-		// to create the actor source in the future anyway, if needed.
-		if (const UTakeRecorderSources* Sources = InSource->GetTypedOuter<UTakeRecorderSources>())
-		{
-			return Sources->GetSettings().bRecordToPossessable;
-		}
-		
-		return false;
-	}
 
 	static bool IsActorRecordable(const AActor* Actor)
 	{
diff --git a/Engine/Plugins/VirtualProduction/Takes/Source/TakeRecorderSources/Public/TakeRecorderActorSource.h b/Engine/Plugins/VirtualProduction/Takes/Source/TakeRecorderSources/Public/TakeRecorderActorSource.h
--- a/Engine/Plugins/VirtualProduction/Takes/Source/TakeRecorderSources/Public/TakeRecorderActorSource.h
+++ b/Engine/Plugins/VirtualProduction/Takes/Source/TakeRecorderSources/Public/TakeRecorderActorSource.h
@@ -51,6 +51,12 @@ public:
 	UPROPERTY(EditAnywhere, BlueprintReadWrite, Category = "Actor Source")
 	ETakeRecorderActorRecordType RecordType;
 
+	/**
+	 * Should recording this actor as a spawnable overwrite existing spawnables or not.
+	 */
+	UPROPERTY(EditAnywhere, BlueprintReadWrite, Category = "Actor Source")
+	ETakeRecorderSpawnableOverwriteBehavior SpawnableOverwriteBehavior;
+	
 	/** 
 	 * Whether to ensure that the parent hierarchy is also recorded. If recording to possessable and the parent is not recorded, 
 	 * the recorded transforms will be in local space since the child will still be attached to the parent in the level after 
@@ -157,6 +163,7 @@ public:
 	virtual FMovieSceneSequenceID GetSequenceID() const override { return TargetSequenceID; }
 	UE_API virtual const IMovieSceneTrackRecorderHost* AsTrackRecorderHost() const override;
 	UE_API virtual bool IsActorSourceForActor(const AActor* Actor) const override;
+	UE_API virtual bool OverwriteExistingTrackData() const override;
 	// ~UTakeRecorderSource Interface
 
 	/** Set the Target actor that we are going to record. Will reset the Recorded Property Map to defaults. */
@@ -282,8 +289,6 @@ protected:
 	UE_API FTransform GetRecordedActorAnimationInitialRootTransform(class AActor* OtherActor) const override;
 	/** Returns the root level sequence being recorded into */
 	ULevelSequence* GetRootLevelSequence() const override { return RootLevelSequence; };
-	/** Returns if an actor is possesable. */
-	UE_API virtual bool IsOtherActorPossessable(AActor* InActor) const override;
 	// ~IMovieSceneTrackRecorderHost Interface
 	
 	/** Initializes an instance of the specified class if we don't already have it in our Settings array. */
diff --git a/Engine/Plugins/VirtualProduction/Takes/Source/TakeTrackRecorders/Private/MovieSceneActorReferenceTrackRecorder.cpp b/Engine/Plugins/VirtualProduction/Takes/Source/TakeTrackRecorders/Private/MovieSceneActorReferenceTrackRecorder.cpp
--- a/Engine/Plugins/VirtualProduction/Takes/Source/TakeTrackRecorders/Private/MovieSceneActorReferenceTrackRecorder.cpp
+++ b/Engine/Plugins/VirtualProduction/Takes/Source/TakeTrackRecorders/Private/MovieSceneActorReferenceTrackRecorder.cpp
@@ -253,40 +253,26 @@ FMovieSceneActorReferenceKey UMovieSceneActorReferenceTrackRecorder::ActorToRefe
 		
 		FMovieSceneObjectBindingID Binding;
 		
-		const TOptional<FMovieSceneSpawnableAnnotation> Spawnable = FMovieSceneSpawnableAnnotation::Find(InActor);
-		if (Spawnable.IsSet())
+		ULevelSequence* TargetLevelSequence = OwningTakeRecorderSource->GetLevelSequence(InActor);
+		if (ensure(TargetLevelSequence))
 		{
-			// @TODO: This pathway is designed for spawnables when recording over an existing recording, but hasn't been fully implemented yet.
-			FText OutError;
-			if (const TSharedPtr<ISequencer> SequencerPtr = TakesUtils::OpenSequencer(RootLevelSequence, &OutError))
+			FGuid ParentActorId = TakesUtils::ResolveActorFromSequence(InActor, TargetLevelSequence);
+			if (!ParentActorId.IsValid())
 			{
-				// Check whether the spawnable is underneath the current sequence, if so, we can remap it to a local sequence ID
-				Binding = UE::MovieScene::FRelativeObjectBindingID(SequencerPtr->GetFocusedTemplateID(), Spawnable->SequenceID, Spawnable->ObjectBindingID, *SequencerPtr);
+				UTakeRecorderSources* Sources = TargetLevelSequence->FindOrAddMetaData<UTakeRecorderSources>();
+				TargetLevelSequence = Sources->AddSourceForActorBinding(InActor, TargetLevelSequence);
+				check(TargetLevelSequence);
+				ParentActorId = TakesUtils::ResolveActorFromSequence(InActor, TargetLevelSequence);
 			}
-		}
-		else
-		{
-			ULevelSequence* TargetLevelSequence = OwningTakeRecorderSource->GetLevelSequence(InActor);
-			if (ensure(TargetLevelSequence))
+		
+			// Actor reference tracks should always be able to find a valid guid -- all sources must be established by this point.
+			if (ensure(ParentActorId.IsValid()))
 			{
-				FGuid ParentActorId = TakesUtils::ResolveActorFromSequence(InActor, TargetLevelSequence);
-				if (!ParentActorId.IsValid())
-				{
-					UTakeRecorderSources* Sources = TargetLevelSequence->FindOrAddMetaData<UTakeRecorderSources>();
-					TargetLevelSequence = Sources->AddSourceForActorBinding(InActor, TargetLevelSequence);
-					check(TargetLevelSequence);
-					ParentActorId = TakesUtils::ResolveActorFromSequence(InActor, TargetLevelSequence);
-				}
-			
-				// Actor reference tracks should always be able to find a valid guid -- all sources must be established by this point.
-				if (ensure(ParentActorId.IsValid()))
-				{
-					const FMovieSceneSequenceID TargetSequenceID = OwningTakeRecorderSource->GetLevelSequenceID(InActor);
-					const FMovieSceneSequenceID ThisSequenceID   = OwningTakeRecorderSource->GetSequenceID();
-			
-					// Setup binding relative to a target sequence, which may or may not be the root sequence, so we can support references in subsequences.
-					Binding = UE::MovieScene::FRelativeObjectBindingID(ThisSequenceID, TargetSequenceID, ParentActorId, OwningTakeRecorderSource->GetRootLevelSequence());
-				}
+				const FMovieSceneSequenceID TargetSequenceID = OwningTakeRecorderSource->GetLevelSequenceID(InActor);
+				const FMovieSceneSequenceID ThisSequenceID   = OwningTakeRecorderSource->GetSequenceID();
+		
+				// Setup binding relative to a target sequence, which may or may not be the root sequence, so we can support references in subsequences.
+				Binding = UE::MovieScene::FRelativeObjectBindingID(ThisSequenceID, TargetSequenceID, ParentActorId, OwningTakeRecorderSource->GetRootLevelSequence());
 			}
 		}
 
diff --git a/Engine/Plugins/VirtualProduction/Takes/Source/TakeTrackRecorders/Public/TrackRecorders/IMovieSceneTrackRecorderHost.h b/Engine/Plugins/VirtualProduction/Takes/Source/TakeTrackRecorders/Public/TrackRecorders/IMovieSceneTrackRecorderHost.h
--- a/Engine/Plugins/VirtualProduction/Takes/Source/TakeTrackRecorders/Public/TrackRecorders/IMovieSceneTrackRecorderHost.h
+++ b/Engine/Plugins/VirtualProduction/Takes/Source/TakeTrackRecorders/Public/TrackRecorders/IMovieSceneTrackRecorderHost.h
@@ -194,9 +194,4 @@ public:
 
 	/**/
 	virtual FMovieSceneSequenceID GetSequenceID() const = 0;
-	
-	/** 
-	 * If the given actor is possessable.
-	 */
-	virtual bool IsOtherActorPossessable(AActor* InActor) const { return false; }
 };
diff --git a/Engine/Plugins/VirtualProduction/Takes/Source/TakesCore/Private/TakeRecorderSource.cpp b/Engine/Plugins/VirtualProduction/Takes/Source/TakesCore/Private/TakeRecorderSource.cpp
--- a/Engine/Plugins/VirtualProduction/Takes/Source/TakesCore/Private/TakeRecorderSource.cpp
+++ b/Engine/Plugins/VirtualProduction/Takes/Source/TakesCore/Private/TakeRecorderSource.cpp
@@ -23,6 +23,11 @@ void UTakeRecorderSource::PostEditChangeProperty(FPropertyChangedEvent& Property
 }
 #endif
 
+bool UTakeRecorderSource::OverwriteExistingTrackData() const
+{
+	return true;
+}
+
 const FSlateBrush* UTakeRecorderSource::GetDisplayIconImpl() const
 {
 	return FSlateIconFinder::FindCustomIconBrushForClass(GetClass(), TEXT("ClassThumbnail"));
diff --git a/Engine/Plugins/VirtualProduction/Takes/Source/TakesCore/Private/TakeRecorderSources.cpp b/Engine/Plugins/VirtualProduction/Takes/Source/TakesCore/Private/TakeRecorderSources.cpp
--- a/Engine/Plugins/VirtualProduction/Takes/Source/TakesCore/Private/TakeRecorderSources.cpp
+++ b/Engine/Plugins/VirtualProduction/Takes/Source/TakesCore/Private/TakeRecorderSources.cpp
@@ -228,6 +228,7 @@ void UTakeRecorderSources::PreRecordingRecursive(TArray<UTakeRecorderSource*> In
 			// is using which sub-sequence so that we can push the correct sequence for all points of the Source's recording lifecycle.
 			if (Settings.bRecordSourcesIntoSubSequences && Source->SupportsSubscenes())
 			{
+				const FString SubSequenceAssetName = ObjectTools::SanitizeObjectName(Source->GetSubsceneAssetName(InRootSequence));
 				const FString BaseSubSequenceTrackName = ObjectTools::SanitizeObjectName(Source->GetSubsceneTrackName(InRootSequence));
 				FString SubSequenceTrackName = BaseSubSequenceTrackName;
 				
@@ -247,8 +248,6 @@ void UTakeRecorderSources::PreRecordingRecursive(TArray<UTakeRecorderSource*> In
 					while (bTrackNameAlreadyUsed);
 				}
 				
-				const FString& SubSequenceAssetName = ObjectTools::SanitizeObjectName(Source->GetSubsceneAssetName(InRootSequence));
-
 				TargetSequence = CreateSubSequenceForSource(InRootSequence, SubSequenceTrackName, SubSequenceAssetName, Settings);
 
 				// If there's already a Subscene Track for our sub-sequence we need to remove that track before create a new one. No data is lost in this process as the
@@ -256,14 +255,22 @@ void UTakeRecorderSources::PreRecordingRecursive(TArray<UTakeRecorderSource*> In
 				const FString SequenceName = FPaths::GetBaseFilename(TargetSequence->GetPathName());
 				UMovieSceneSubTrack* SubsceneTrack = nullptr;
 
-				for (UMovieSceneTrack* Track : InRootSequence->GetMovieScene()->GetTracks())
+				if (Source->OverwriteExistingTrackData())
 				{
-					if (Track->IsA<UMovieSceneSubTrack>())
+					for (UMovieSceneTrack* Track : InRootSequence->GetMovieScene()->GetTracks())
 					{
-						if (Track->GetDisplayName().ToString() == SubSequenceTrackName)
+						if (Track->IsA<UMovieSceneSubTrack>())
 						{
-							SubsceneTrack = CastChecked<UMovieSceneSubTrack>(Track);
-							SubsceneTrack->RemoveAllAnimationData();
+							// TODO: Need something other than DisplayName to track this.
+							// Do we want to increment the name count when not overwriting?
+							if (Track->GetDisplayName().ToString() == SubSequenceTrackName)
+							{
+								SubsceneTrack = CastChecked<UMovieSceneSubTrack>(Track);
+								SubsceneTrack->RemoveAllAnimationData();
+								// Currently, we only expect 1 copy at most. If we want to overwrite all copies
+								// there needs to be additional logic to clean out other tracks.
+								break;
+							}
 						}
 					}
 				}
@@ -682,6 +689,24 @@ void UTakeRecorderSources::StopRecording(class ULevelSequence* InSequence, const
 	}
 }
 
+ULevelSequence* UTakeRecorderSources::FindExistingSubSequence(const ULevelSequence* InRootSequence, const FString& InSubSequenceAssetName)
+{
+	if (const UMovieSceneSubTrack* SubTrack = InRootSequence->GetMovieScene()->FindTrack<UMovieSceneSubTrack>())
+	{
+		// Look at each section in the track to see if it has the same name as our new SubSequence name.
+		for (UMovieSceneSection* Section : SubTrack->GetAllSections())
+		{
+			const UMovieSceneSubSection* SubSection = CastChecked<UMovieSceneSubSection>(Section);
+			if (FPaths::GetBaseFilename(SubSection->GetSequence()->GetPathName()) == InSubSequenceAssetName)
+			{
+				return CastChecked<ULevelSequence>(SubSection->GetSequence());
+			}
+		}
+	}
+
+	return nullptr;
+}
+
 ULevelSequence* UTakeRecorderSources::CreateSubSequenceForSource(ULevelSequence* InRootSequence, const FString& SubSequenceTrackName,
 	const FString& SubSequenceAssetName, const FTakeRecorderSourcesSettings& InSettings)
 {
@@ -689,27 +714,16 @@ ULevelSequence* UTakeRecorderSources::CreateSubSequenceForSource(ULevelSequence*
 
 	// We want to sanitize the object names because sometimes they come from names with spaces and other invalid characters in them.
 	const FString& SequenceDirectory = FPaths::GetPath(InRootSequence->GetPathName());
-	const FString& SequenceName = FPaths::GetBaseFilename(InRootSequence->GetPathName());
+	FString SequenceName = FPaths::GetBaseFilename(InRootSequence->GetPathName());
 
 	// We need to check the Root Sequence to see if they already have a sub-sequence with this name so that we duplicate the right
 	// sequence and re-use that, instead of just making a new blank sequence every time. This will help in cases where they've done a recording, 
 	// modified a sub-sequence and want to record from that setup. Each source will individually remove any old data inside the Sub Sequence
 	// so we don't have to worry about any data the user added via Sequencer unrelated to what they recorded.
-	ULevelSequence* ExistingSubSequence = nullptr;
-	UMovieSceneSubTrack* SubTrack = InRootSequence->GetMovieScene()->FindTrack<UMovieSceneSubTrack>();
-	if (SubTrack)
+	ULevelSequence* ExistingSubSequence = FindExistingSubSequence(InRootSequence, SubSequenceAssetName);
+	if (ExistingSubSequence)
 	{
-		// Look at each section in the track to see if it has the same name as our new SubSequence name.
-		for (UMovieSceneSection* Section : SubTrack->GetAllSections())
-		{
-			UMovieSceneSubSection* SubSection = CastChecked<UMovieSceneSubSection>(Section);
-			if (FPaths::GetBaseFilename(SubSection->GetSequence()->GetPathName()) == SubSequenceAssetName)
-			{
-				UE_LOG(LogTakesCore, Log, TEXT("Found existing sub-section for source %s, duplicating sub-section for recording into."), *SubSequenceAssetName);
-				ExistingSubSequence = CastChecked<ULevelSequence>(SubSection->GetSequence());
-				break;
-			}
-		}
+		UE_LOG(LogTakesCore, Log, TEXT("Found existing sub-section for source %s, duplicating sub-section for recording into."), *SubSequenceAssetName);
 	}
 	
 	FString SubSequenceDirectory;
diff --git a/Engine/Plugins/VirtualProduction/Takes/Source/TakesCore/Public/TakeRecorderSource.h b/Engine/Plugins/VirtualProduction/Takes/Source/TakesCore/Public/TakeRecorderSource.h
--- a/Engine/Plugins/VirtualProduction/Takes/Source/TakesCore/Public/TakeRecorderSource.h
+++ b/Engine/Plugins/VirtualProduction/Takes/Source/TakesCore/Public/TakeRecorderSource.h
@@ -178,6 +178,13 @@ public:
 	 */
 	virtual bool CanAddSource(UTakeRecorderSources* InSources) const { return true; }
 
+	/**
+	 * Whether this source should overwrite existing track data, or leave it, or potentially create new data.
+	 * This is primarily used in response to Spawnables, such as we if we want to overwrite subsequence data
+	 * or clear existing data in non-subsequence tracks.
+	 */
+	UE_API virtual bool OverwriteExistingTrackData() const;
+
 public:
 
 	/**
diff --git a/Engine/Plugins/VirtualProduction/Takes/Source/TakesCore/Public/TakeRecorderSources.h b/Engine/Plugins/VirtualProduction/Takes/Source/TakesCore/Public/TakeRecorderSources.h
--- a/Engine/Plugins/VirtualProduction/Takes/Source/TakesCore/Public/TakeRecorderSources.h
+++ b/Engine/Plugins/VirtualProduction/Takes/Source/TakesCore/Public/TakeRecorderSources.h
@@ -9,6 +9,8 @@
 #include "Misc/FrameRate.h"
 #include "Misc/QualifiedFrameTime.h"
 #include "Serializers/MovieSceneManifestSerialization.h"
+#include "TakeRecorderSpawnableBehavior.h"
+
 #include "TakeRecorderSources.generated.h"
 
 #define UE_API TAKESCORE_API
@@ -23,6 +25,7 @@ struct FTakeRecorderSourcesSettings
 	bool bStartAtCurrentTimecode;
 	bool bRecordSourcesIntoSubSequences;
 	bool bRecordToPossessable;
+	ETakeRecorderSpawnableOverwriteBehavior SpawnableOverwriteBehavior;
 	bool bSaveRecordedAssets;
 	bool bAutoLock;
 	bool bRemoveRedundantTracks;
@@ -221,6 +224,8 @@ public:
 	FQualifiedFrameTime GetCachedFrameTime() const { return CachedFrameTime; }
 
 private:
+	static UE_API ULevelSequence* FindExistingSubSequence(const ULevelSequence* InRootSequence, const FString& InSubSequenceAssetName);
+	
 	/**
 	 * Renders stat text that displays the current list of active take recorder sources.  The take recorder can now operate without a panel and this
 	 * gives users a way to view the current sources list outside of the panel.
diff --git a/Engine/Plugins/VirtualProduction/Takes/Source/TakesCore/Public/TakeRecorderSpawnableBehavior.h b/Engine/Plugins/VirtualProduction/Takes/Source/TakesCore/Public/TakeRecorderSpawnableBehavior.h
new file mode 100644
--- /dev/null
+++ b/Engine/Plugins/VirtualProduction/Takes/Source/TakesCore/Public/TakeRecorderSpawnableBehavior.h
@@ -0,0 +1,18 @@
+// Copyright Epic Games, Inc. All Rights Reserved.
+
+#pragma once
+
+#include "HAL/Platform.h"
+
+#include "TakeRecorderSpawnableBehavior.generated.h"
+
+UENUM(BlueprintType)
+enum class ETakeRecorderSpawnableOverwriteBehavior : uint8
+{
+	/** Use the project default to determine if spawnables are overwritten */
+	ProjectDefault,
+	/** Overwrite one existing spawnable similar to behavior in UE 5.7 and below */
+	OverwriteLegacy UMETA(DisplayName = "Overwrite (Legacy)"),
+	/** Creates a new spawnable */
+	CreateNew
+};
\ No newline at end of file
