TakeRecorder - Add in Take Recorder Sources Manager interface to hide away the sources logic for runtime and editor Take Recording.

The new ITakeRecorderSourcesManager interface is added and implemented. No code is yet changed to use this new interface due to other module restructures that are required, but this adds the framework and structure to then easily move to.
At editor time, this implementation will use the LevelSequence metadata as before - there is no change to behaviour in this context.
At runtime, the implementation, a subsystem, keeps its own store of the take recorder sources keyed by LevelSequence. These can then be modified and queried as needed.

#Synced-CL 52293313

diff --git a/Engine/Plugins/VirtualProduction/Takes/Source/TakesCore/Private/ITakeRecorderSourcesManager.cpp b/Engine/Plugins/VirtualProduction/Takes/Source/TakesCore/Private/ITakeRecorderSourcesManager.cpp
new file mode 100644
--- /dev/null
+++ b/Engine/Plugins/VirtualProduction/Takes/Source/TakesCore/Private/ITakeRecorderSourcesManager.cpp
@@ -0,0 +1,17 @@
+// Copyright Epic Games, Inc. All Rights Reserved.
+
+#include "ITakeRecorderSourcesManager.h"
+#include "TakeRecorderSourcesManagerImpl.h"
+#include "Engine/Engine.h"
+
+ITakeRecorderSourcesManager* ITakeRecorderSourcesManager::Get()
+{
+	return GEngine->GetEngineSubsystem<UTakeRecorderSourcesManagerImpl>();
+}
+
+ITakeRecorderSourcesManager& ITakeRecorderSourcesManager::GetChecked()
+{
+	ITakeRecorderSourcesManager* Ptr = Get();
+	check(Ptr);
+	return *Ptr;
+}
diff --git a/Engine/Plugins/VirtualProduction/Takes/Source/TakesCore/Private/TakeRecorderSourcesManagerImpl.cpp b/Engine/Plugins/VirtualProduction/Takes/Source/TakesCore/Private/TakeRecorderSourcesManagerImpl.cpp
new file mode 100644
--- /dev/null
+++ b/Engine/Plugins/VirtualProduction/Takes/Source/TakesCore/Private/TakeRecorderSourcesManagerImpl.cpp
@@ -0,0 +1,124 @@
+// Copyright Epic Games, Inc. All Rights Reserved.
+
+#include "TakeRecorderSourcesManagerImpl.h"
+
+#include "LevelSequence.h"
+#include "TakeRecorderSource.h"
+#include "TakeRecorderSources.h"
+
+void UTakeRecorderSourcesManagerImpl::Initialize(FSubsystemCollectionBase& Collection)
+{
+}
+
+void UTakeRecorderSourcesManagerImpl::Deinitialize()
+{
+	RuntimeSources.Empty();
+}
+
+UTakeRecorderSources* UTakeRecorderSourcesManagerImpl::CopySources(ULevelSequence* FromSequence, ULevelSequence* ToSequence)
+{
+	if (!IsValid(FromSequence) || !IsValid(ToSequence))
+	{
+		return nullptr;
+	}
+
+	if (UTakeRecorderSources* FoundSources = FindSources(FromSequence))
+	{
+	#if WITH_EDITOR
+		// when in Editor, ensure they are added to the metadata, overwriting anything thats already there.
+		return ToSequence->CopyMetaData<UTakeRecorderSources>(FoundSources);
+	#else // WITH_EDITOR
+		UTakeRecorderSources* CopiedSources = DuplicateObject(FoundSources, ToSequence);
+		RuntimeSources.Add(ToSequence, CopiedSources);
+		return CopiedSources;
+	#endif // WITH_EDITOR
+	}
+
+	return nullptr;
+}
+
+UTakeRecorderSources* UTakeRecorderSourcesManagerImpl::FindSources(const ULevelSequence* InSequence) const
+{
+	if (!IsValid(InSequence))
+	{
+		return nullptr;
+	}
+
+#if WITH_EDITOR
+	// When in editor, we check the sequence metadata.
+	if (UTakeRecorderSources* MetaDataSources = InSequence->FindMetaData<UTakeRecorderSources>())
+	{
+		return MetaDataSources;
+	}
+#else // WITH_EDITOR
+	// check if the sources are stored for the sequence
+	if (TObjectPtr<UTakeRecorderSources> const* SourcesPtr = RuntimeSources.Find(InSequence))
+	{
+		return *SourcesPtr;
+	}
+#endif // WITH_EDITOR
+
+	return nullptr;
+}
+
+UTakeRecorderSources* UTakeRecorderSourcesManagerImpl::FindOrAddSources(ULevelSequence* InSequence)
+{
+	if (!IsValid(InSequence))
+	{
+		return nullptr;
+	}
+
+#if WITH_EDITOR
+	return InSequence->FindOrAddMetaData<UTakeRecorderSources>();
+#else // WITH_EDITOR
+	if (UTakeRecorderSources* FoundSources = FindSources(InSequence))
+	{
+		return FoundSources;
+	}
+
+	UTakeRecorderSources* NewSources = NewObject<UTakeRecorderSources>(InSequence);
+	RuntimeSources.Add(InSequence, NewSources);
+	return NewSources;
+#endif // WITH_EDITOR
+}
+
+bool UTakeRecorderSourcesManagerImpl::HasSources(const ULevelSequence* InSequence) const
+{
+	return FindSources(InSequence) != nullptr;
+}
+
+UTakeRecorderSource* UTakeRecorderSourcesManagerImpl::AddSource(ULevelSequence* InSequence, const TSubclassOf<UTakeRecorderSource> InSourceClass)
+{
+	if (UTakeRecorderSources* FoundSources = FindOrAddSources(InSequence))
+	{
+		return FoundSources->AddSource(InSourceClass);
+	}
+	return nullptr;
+}
+
+void UTakeRecorderSourcesManagerImpl::RemoveSources(ULevelSequence* InSequence)
+{
+	if (!IsValid(InSequence))
+	{
+		return;
+	}
+
+#if WITH_EDITOR
+	InSequence->RemoveMetaData<UTakeRecorderSources>();
+#else // WITH_EDITOR
+	RuntimeSources.Remove(InSequence);
+#endif // WITH_EDITOR
+}
+
+void UTakeRecorderSourcesManagerImpl::RemoveSource(ULevelSequence* InSequence, UTakeRecorderSource* InSource)
+{
+	if (!IsValid(InSource))
+	{
+		return;
+	}
+	
+	if (UTakeRecorderSources* FoundSources = FindSources(InSequence))
+	{
+		FoundSources->RemoveSource(InSource);
+	}
+}
diff --git a/Engine/Plugins/VirtualProduction/Takes/Source/TakesCore/Private/TakeRecorderSourcesManagerImpl.h b/Engine/Plugins/VirtualProduction/Takes/Source/TakesCore/Private/TakeRecorderSourcesManagerImpl.h
new file mode 100644
--- /dev/null
+++ b/Engine/Plugins/VirtualProduction/Takes/Source/TakesCore/Private/TakeRecorderSourcesManagerImpl.h
@@ -0,0 +1,49 @@
+// Copyright Epic Games, Inc. All Rights Reserved.
+
+#pragma once
+
+#include "Containers/Map.h"
+#include "Subsystems/EngineSubsystem.h"
+#include "ITakeRecorderSourcesManager.h"
+
+#include "TakeRecorderSourcesManagerImpl.generated.h"
+
+/**
+ * Implementation of the ITakeRecorderSourcesManager interface, implemented as an Engine subsystem.
+ * Handles Take Recorder Sources in both editor and runtime contexts.
+ * At editor time, the LevelSequence metadata is used to track and store the Take Recorder sources.
+ * At runtime, a transient map is used to store the sources per level sequence. This is designed to be transient and unserialized due to
+ * requirements of recording at runtime, and setups most likely be game based.
+ */
+UCLASS(MinimalAPI)
+class UTakeRecorderSourcesManagerImpl final : 
+	public UEngineSubsystem, 
+	public ITakeRecorderSourcesManager
+{
+	GENERATED_BODY()
+
+public:
+	//~ Begin UEngineSubsystem
+	virtual void Initialize(FSubsystemCollectionBase& Collection) override;
+	virtual void Deinitialize() override;
+	//~ End UEngineSubsystem
+
+	//~ Begin ITakeRecorderSourcesManager
+	virtual UTakeRecorderSources* CopySources(ULevelSequence* FromSequence, ULevelSequence* ToSequence) override;
+	virtual UTakeRecorderSources* FindSources(const ULevelSequence* InSequence) const override;
+	virtual UTakeRecorderSources* FindOrAddSources(ULevelSequence* InSequence) override;
+	virtual bool HasSources(const ULevelSequence* InSequence) const override;
+	virtual UTakeRecorderSource* AddSource(ULevelSequence* InSequence, const TSubclassOf<UTakeRecorderSource> InSourceClass) override;
+	virtual void RemoveSource(ULevelSequence* InSequence, UTakeRecorderSource* InSource) override;
+	virtual void RemoveSources(ULevelSequence* InSequence) override;
+	//~ End ITakeRecorderSourcesManager
+
+private:
+	/** 
+	 * Sources map used at runtime to store the TakeRecorderSources per LevelSequence.
+	 * This map can be transient as we have no intention of ever serializing it.
+	 * It is used only at runtime for adding sources to record, and does not need to be persisted out.
+	 */
+	UPROPERTY(Transient)
+	TMap<TObjectPtr<ULevelSequence>, TObjectPtr<UTakeRecorderSources>> RuntimeSources;
+};
diff --git a/Engine/Plugins/VirtualProduction/Takes/Source/TakesCore/Public/ITakeRecorderSourcesManager.h b/Engine/Plugins/VirtualProduction/Takes/Source/TakesCore/Public/ITakeRecorderSourcesManager.h
new file mode 100644
--- /dev/null
+++ b/Engine/Plugins/VirtualProduction/Takes/Source/TakesCore/Public/ITakeRecorderSourcesManager.h
@@ -0,0 +1,108 @@
+// Copyright Epic Games, Inc. All Rights Reserved.
+
+#pragma once
+
+#include "UObject/Interface.h"
+
+#include "ITakeRecorderSourcesManager.generated.h"
+
+class UTakeRecorderSource;
+class UTakeRecorderSources;
+class ULevelSequence;
+
+UINTERFACE(MinimalAPI, NotBlueprintable)
+class UTakeRecorderSourcesManager : public UInterface
+{
+	GENERATED_BODY()
+};
+
+/**
+ * Interface for the public Take Recorder Sources manager interface.
+ * This interface allows callers to work with Take Recorder Sources uniformly at editor and runtime,
+ * with the implementation handling the differences between the two environments.
+ * 
+ * At editor time, the LevelSequence metadata is used to store and track sources as is the traditional behaviour.
+ * At runtime, the sources should be stored per level sequence editor separate from the sequence, as metadata is unavailable in runtime.
+ */
+class ITakeRecorderSourcesManager
+{
+	GENERATED_BODY()
+
+public:
+	/** Returns a pointer to the global ITakeRecorderSourcesManager implementation, or nullptr if not available. */
+	static TAKESCORE_API ITakeRecorderSourcesManager* Get();
+
+	/**
+	 * Returns a reference to the global ITakeRecorderSourcesManager implementation.
+	 * Asserts if no implementation is available.
+	 */
+	static TAKESCORE_API ITakeRecorderSourcesManager& GetChecked();
+
+	/**
+	 * Copies all recorder sources from one level sequence to another.
+	 *
+	 * @param FromSequence  The sequence to copy sources from.
+	 * @param ToSequence    The sequence to copy sources to.
+	 * @return Pointer to the newly copied Sources object, or nullptr if sources were not found in FromSequence, or nothing was copied.
+	 */
+	virtual UTakeRecorderSources* CopySources(ULevelSequence* FromSequence, ULevelSequence* ToSequence) = 0;
+
+	/**
+	 * Returns the UTakeRecorderSources associated with the given level sequence, or nullptr if none exist.
+	 *
+	 * @param InSequence  The sequence to look up.
+	 */
+	virtual UTakeRecorderSources* FindSources(const ULevelSequence* InSequence) const = 0;
+
+	/**
+	 * Returns the UTakeRecorderSources associated with the given level sequence, creating and associating
+	 * a new instance if one does not already exist.
+	 *
+	 * @param InSequence  The sequence to look up or create sources for.
+	 */
+	virtual UTakeRecorderSources* FindOrAddSources(ULevelSequence* InSequence) = 0;
+
+	/**
+	 * Returns true if a UTakeRecorderSources object is currently associated with the given level sequence.
+	 *
+	 * @param InSequence  The sequence to check.
+	 */
+	virtual bool HasSources(const ULevelSequence* InSequence) const = 0;
+
+	/**
+	 * Convenience wrapper that adds a new recorder source of type TSource to the given level sequence.
+	 *
+	 * @param InSequence  The sequence to add the source to.
+	 * @return The newly created UTakeRecorderSource, or nullptr on failure.
+	 */
+	template <typename TSource>
+	UTakeRecorderSource* AddSource(ULevelSequence* InSequence)
+	{
+		return AddSource(TSource::StaticClass(), InSequence);
+	}
+
+	/**
+	 * Adds a new recorder source of the specified class to the given level sequence.
+	 *
+	 * @param InSequence     The sequence to add the source to.
+	 * @param InSourceClass  The class of source to create and add.
+	 * @return The newly created UTakeRecorderSource, or nullptr on failure.
+	 */
+	virtual UTakeRecorderSource* AddSource(ULevelSequence* InSequence, const TSubclassOf<UTakeRecorderSource> InSourceClass) = 0;
+
+	/**
+	 * Removes a recorder source from the given level sequence's source list.
+	 *
+	 * @param InSequence  The sequence whose source list to modify.
+	 * @param InSource    The source to remove.
+	 */
+	virtual void RemoveSource(ULevelSequence* InSequence, UTakeRecorderSource* InSource) = 0;
+
+	/**
+	 * Removes the UTakeRecorderSources association for the given level sequence.
+	 * Callers are responsible for invoking this when a sequence is no longer needed.
+	 *
+	 * @param InSequence  The sequence whose sources entry should be removed.
+	 */
+	virtual void RemoveSources(ULevelSequence* InSequence) = 0;
+};
