Take Recorder - Change Take Recorder sources to use Manager interface instead of Level Sequence MetaData directly.

Now that the ITakeRecorderSourcesManager is now available, updating all places that use the LevelSequence metadata directly in editor to work with sources, to now use the interface. This will mean that during runtime they will correctly handle sources without the use of LevelSequence MetaData.

#rnx
#Synced-CL 52338471

diff --git a/Engine/Plugins/VirtualProduction/Takes/Source/TakeRecorder/Private/Recorder/TakeRecorder.cpp b/Engine/Plugins/VirtualProduction/Takes/Source/TakeRecorder/Private/Recorder/TakeRecorder.cpp
--- a/Engine/Plugins/VirtualProduction/Takes/Source/TakeRecorder/Private/Recorder/TakeRecorder.cpp
+++ b/Engine/Plugins/VirtualProduction/Takes/Source/TakeRecorder/Private/Recorder/TakeRecorder.cpp
@@ -30,6 +30,7 @@
 #include "Tickable.h"
 #include "Stats/Stats.h"
 #include "SequencerSettings.h"
+#include "ITakeRecorderSourcesManager.h"
 
 // Engine includes
 #include "GameFramework/WorldSettings.h"
@@ -853,7 +854,7 @@ void UTakeRecorder::InternalTick(float DeltaTime)
 
 	FQualifiedFrameTime RecordTime = GetRecordTime();
 	
-	UTakeRecorderSources* Sources = SequenceAsset->FindOrAddMetaData<UTakeRecorderSources>();
+	UTakeRecorderSources* Sources = ITakeRecorderSourcesManager::GetChecked().FindOrAddSources(SequenceAsset);
 	if (!Parameters.bDisableRecordingAndSave)
 	{
 		CurrentFrameTime = Sources->TickRecording(SequenceAsset, RecordTime, DeltaTime);
@@ -892,7 +893,7 @@ void UTakeRecorder::PreRecord()
 	
 	State = ETakeRecorderState::PreRecord;
 
-	UTakeRecorderSources* Sources = SequenceAsset->FindMetaData<UTakeRecorderSources>();
+	UTakeRecorderSources* Sources = ITakeRecorderSourcesManager::GetChecked().FindSources(SequenceAsset);
 	check(Sources);
 	UTakeMetaData* AssetMetaData = SequenceAsset->FindMetaData<UTakeMetaData>();
 
@@ -1061,7 +1062,7 @@ void UTakeRecorder::Start()
 			Sequencer->SetPlaybackStatus(EMovieScenePlayerStatus::Playing);
 		}
 	}
-	UTakeRecorderSources* Sources = SequenceAsset->FindMetaData<UTakeRecorderSources>();
+	UTakeRecorderSources* Sources = ITakeRecorderSourcesManager::GetChecked().FindSources(SequenceAsset);
 	check(Sources);
 
 	UTakeMetaData* AssetMetaData = SequenceAsset->FindMetaData<UTakeMetaData>();
@@ -1201,7 +1202,7 @@ void UTakeRecorder::StopInternal(const bool bCancelled)
 
 		if (!Parameters.bDisableRecordingAndSave)
 		{
-			UTakeRecorderSources* Sources = SequenceAsset->FindMetaData<UTakeRecorderSources>();
+			UTakeRecorderSources* Sources = ITakeRecorderSourcesManager::GetChecked().FindSources(SequenceAsset);
 			check(Sources);
 			Sources->StopRecording(SequenceAsset, bCancelled);
 		}
diff --git a/Engine/Plugins/VirtualProduction/Takes/Source/TakeRecorder/Private/Recorder/TakeRecorderSubsystemImplementation.cpp b/Engine/Plugins/VirtualProduction/Takes/Source/TakeRecorder/Private/Recorder/TakeRecorderSubsystemImplementation.cpp
--- a/Engine/Plugins/VirtualProduction/Takes/Source/TakeRecorder/Private/Recorder/TakeRecorderSubsystemImplementation.cpp
+++ b/Engine/Plugins/VirtualProduction/Takes/Source/TakeRecorder/Private/Recorder/TakeRecorderSubsystemImplementation.cpp
@@ -26,6 +26,7 @@
 #include "TakesCoreBlueprintLibrary.h"
 #include "Widgets/Notifications/SNotificationList.h"
 #include "Widgets/TakeRecorderNamingTokensData.h"
+#include "ITakeRecorderSourcesManager.h"
 
 #include UE_INLINE_GENERATED_CPP_BY_NAME(TakeRecorderSubsystemImplementation)
 
@@ -172,7 +173,7 @@ void UTakeRecorderSubsystemImplementation::SetTargetSequence(const FTakeRecorder
 
 		ULevelSequence* LevelSequence = TransientPreset->GetLevelSequence();
 
-		UTakeRecorderSources* BaseSources = InData.BaseSequence->FindMetaData<UTakeRecorderSources>();
+		UTakeRecorderSources* BaseSources = ITakeRecorderSourcesManager::GetChecked().FindSources(InData.BaseSequence);
 		if (BaseSources && LevelSequence)
 		{
 			LevelSequence->CopyMetaData(BaseSources);
@@ -287,7 +288,7 @@ bool UTakeRecorderSubsystemImplementation::StartRecording(bool bOpenSequencer, b
 	TGuardValue<bool> ReentrantGuard(bStartedRecording, true);
 
 	ULevelSequence*       LevelSequence = GetLevelSequence();
-	UTakeRecorderSources* Sources = LevelSequence ? LevelSequence->FindMetaData<UTakeRecorderSources>() : nullptr;
+	UTakeRecorderSources* Sources = ITakeRecorderSourcesManager::GetChecked().FindSources(LevelSequence);
 	
 	if (LevelSequence && Sources)
 	{
@@ -409,7 +410,7 @@ void UTakeRecorderSubsystemImplementation::ClearPendingTake()
 
 	if (const ULevelSequence* CurrentLevelSequence = GetLevelSequence())
 	{
-		BaseSources = CurrentLevelSequence->FindMetaData<UTakeRecorderSources>();
+		BaseSources = ITakeRecorderSourcesManager::GetChecked().FindSources(CurrentLevelSequence);
 	}
 
 	SuppliedLevelSequence = nullptr;
@@ -491,7 +492,7 @@ void UTakeRecorderSubsystemImplementation::ClearSources()
 UTakeRecorderSources* UTakeRecorderSubsystemImplementation::GetSources() const
 {
 	ULevelSequence* LevelSequence = GetLevelSequence();
-	UTakeRecorderSources* Sources = LevelSequence ? LevelSequence->FindOrAddMetaData<UTakeRecorderSources>() : nullptr;
+	UTakeRecorderSources* Sources = ITakeRecorderSourcesManager::GetChecked().FindOrAddSources(LevelSequence);
 	return Sources;
 }
 
@@ -692,7 +693,7 @@ bool UTakeRecorderSubsystemImplementation::MarkFrame()
 
 		FMovieSceneMarkedFrame MarkedFrame;
 
-		const UTakeRecorderSources* Sources = LevelSequence->FindMetaData<UTakeRecorderSources>();
+		const UTakeRecorderSources* Sources = ITakeRecorderSourcesManager::GetChecked().FindSources(LevelSequence);
 		if (Sources && Sources->GetSettings().bStartAtCurrentTimecode)
 		{
 			MarkedFrame.FrameNumber = FFrameRate::TransformTime(FFrameTime(FApp::GetTimecode().ToFrameNumber(DisplayRate)), DisplayRate, TickResolution).FloorToFrame();
diff --git a/Engine/Plugins/VirtualProduction/Takes/Source/TakeRecorder/Private/Widgets/SLevelSequenceTakeEditor.cpp b/Engine/Plugins/VirtualProduction/Takes/Source/TakeRecorder/Private/Widgets/SLevelSequenceTakeEditor.cpp
--- a/Engine/Plugins/VirtualProduction/Takes/Source/TakeRecorder/Private/Widgets/SLevelSequenceTakeEditor.cpp
+++ b/Engine/Plugins/VirtualProduction/Takes/Source/TakeRecorder/Private/Widgets/SLevelSequenceTakeEditor.cpp
@@ -17,6 +17,7 @@
 #include "Widgets/STakeRecorderSources.h"
 #include "TakeRecorderModule.h"
 #include "LevelSequence.h"
+#include "ITakeRecorderSourcesManager.h"
 
 // Core includes
 #include "Modules/ModuleManager.h"
@@ -195,7 +196,7 @@ void SLevelSequenceTakeEditor::CheckForNewLevelSequence()
 	{
 		CachedLevelSequence = NewLevelSequence;
 
-		UTakeRecorderSources* Sources = NewLevelSequence ? NewLevelSequence->FindOrAddMetaData<UTakeRecorderSources>() : nullptr;
+		UTakeRecorderSources* Sources = ITakeRecorderSourcesManager::GetChecked().FindOrAddSources(NewLevelSequence);
 
 		SourcesWidget->SetSourceObject(Sources);
 		bRequestDetailsRefresh = true;
@@ -227,7 +228,7 @@ TSharedRef<SWidget> SLevelSequenceTakeEditor::OnGenerateSourcesMenu()
 	TSharedRef<FExtender> Extender = MakeShared<FExtender>();
 	{
 		ULevelSequence*       LevelSequence = LevelSequenceAttribute.Get();
-		UTakeRecorderSources* Sources       = LevelSequence ? LevelSequence->FindOrAddMetaData<UTakeRecorderSources>() : nullptr;
+		UTakeRecorderSources* Sources       = ITakeRecorderSourcesManager::GetChecked().FindOrAddSources(LevelSequence);
 
 		if (Sources)
 		{
@@ -268,7 +269,7 @@ TSharedRef<SWidget> SLevelSequenceTakeEditor::OnGenerateSourcesMenu()
 void SLevelSequenceTakeEditor::AddSourceFromClass(TSubclassOf<UTakeRecorderSource> SourceClass)
 {
 	ULevelSequence*       LevelSequence = LevelSequenceAttribute.Get();
-	UTakeRecorderSources* Sources       = LevelSequence ? LevelSequence->FindOrAddMetaData<UTakeRecorderSources>() : nullptr;
+	UTakeRecorderSources* Sources       = ITakeRecorderSourcesManager::GetChecked().FindOrAddSources(LevelSequence);
 
 	if (*SourceClass && Sources)
 	{
@@ -282,7 +283,7 @@ void SLevelSequenceTakeEditor::AddSourceFromClass(TSubclassOf<UTakeRecorderSourc
 bool SLevelSequenceTakeEditor::CanAddSourceFromClass(TSubclassOf<UTakeRecorderSource> SourceClass)
 {
 	ULevelSequence*       LevelSequence = LevelSequenceAttribute.Get();
-	UTakeRecorderSources* Sources = LevelSequence ? LevelSequence->FindOrAddMetaData<UTakeRecorderSources>() : nullptr;
+	UTakeRecorderSources* Sources = ITakeRecorderSourcesManager::GetChecked().FindOrAddSources(LevelSequence);
 
 	if (*SourceClass && Sources)
 	{
diff --git a/Engine/Plugins/VirtualProduction/Takes/Source/TakeRecorder/Private/Widgets/STakeRecorderCockpit.cpp b/Engine/Plugins/VirtualProduction/Takes/Source/TakeRecorder/Private/Widgets/STakeRecorderCockpit.cpp
--- a/Engine/Plugins/VirtualProduction/Takes/Source/TakeRecorder/Private/Widgets/STakeRecorderCockpit.cpp
+++ b/Engine/Plugins/VirtualProduction/Takes/Source/TakeRecorder/Private/Widgets/STakeRecorderCockpit.cpp
@@ -22,6 +22,7 @@
 #include "LevelSequence.h"
 #include "Algo/Find.h"
 #include "Framework/MultiBox/MultiBoxBuilder.h"
+#include "ITakeRecorderSourcesManager.h"
 
 #include "MovieScene.h"
 #include "MovieSceneToolsProjectSettings.h"
@@ -544,7 +545,7 @@ void STakeRecorderCockpit::UpdateRecordError()
 		return;
 	}
 
-	UTakeRecorderSources* SourcesContainer = Sequence->FindMetaData<UTakeRecorderSources>();
+	UTakeRecorderSources* SourcesContainer = ITakeRecorderSourcesManager::GetChecked().FindSources(Sequence);
 	bool bValidSources = false;
 	if (SourcesContainer)
 	{
@@ -1028,7 +1029,7 @@ bool STakeRecorderCockpit::IsLocked() const
 void STakeRecorderCockpit::OnToggleRecording(ECheckBoxState)
 {
 	ULevelSequence*       LevelSequence = TakeRecorderSubsystem->GetLevelSequence();
-	UTakeRecorderSources* Sources = LevelSequence ? LevelSequence->FindMetaData<UTakeRecorderSources>() : nullptr;
+	UTakeRecorderSources* Sources = ITakeRecorderSourcesManager::GetChecked().FindSources(LevelSequence);
 
 	UTakeRecorder* CurrentRecording = UTakeRecorder::GetActiveRecorder();
 	if (CurrentRecording)
diff --git a/Engine/Plugins/VirtualProduction/Takes/Source/TakeRecorder/Private/Widgets/STakeRecorderTabContent.cpp b/Engine/Plugins/VirtualProduction/Takes/Source/TakeRecorder/Private/Widgets/STakeRecorderTabContent.cpp
--- a/Engine/Plugins/VirtualProduction/Takes/Source/TakeRecorder/Private/Widgets/STakeRecorderTabContent.cpp
+++ b/Engine/Plugins/VirtualProduction/Takes/Source/TakeRecorder/Private/Widgets/STakeRecorderTabContent.cpp
@@ -14,6 +14,7 @@
 #include "LevelSequence.h"
 #include "Recorder/TakeRecorderPanel.h"
 #include "Recorder/TakeRecorderBlueprintLibrary.h"
+#include "ITakeRecorderSourcesManager.h"
 
 // Analytics
 #include "EngineAnalytics.h"
@@ -300,7 +301,7 @@ UTakeRecorderSources* STakeRecorderTabContent::GetSources() const
 		LevelSequence = Panel->GetLevelSequence();
 	}
 
-	return LevelSequence ? LevelSequence->FindMetaData<UTakeRecorderSources>() : nullptr;
+	return ITakeRecorderSourcesManager::GetChecked().FindSources(LevelSequence);
 }
 
 void STakeRecorderTabContent::StartRecording() const
diff --git a/Engine/Plugins/VirtualProduction/Takes/Source/TakeRecorderSources/Private/TakeRecorderSourcesUtils.h b/Engine/Plugins/VirtualProduction/Takes/Source/TakeRecorderSources/Private/TakeRecorderSourcesUtils.h
--- a/Engine/Plugins/VirtualProduction/Takes/Source/TakeRecorderSources/Private/TakeRecorderSourcesUtils.h
+++ b/Engine/Plugins/VirtualProduction/Takes/Source/TakeRecorderSources/Private/TakeRecorderSourcesUtils.h
@@ -18,6 +18,8 @@
 #include "MovieSceneSequenceID.h"
 #include "Compilation/MovieSceneCompiledDataManager.h"
 
+#include "ITakeRecorderSourcesManager.h"
+
 class UWorld;
 class ULevelSequence;
 class UMovieScene;
@@ -38,7 +40,7 @@ namespace TakeRecorderSourcesUtils
 		}
 
 		// Otherwise any of the source's worlds
-		UTakeRecorderSources* Sources = InSequence->FindOrAddMetaData<UTakeRecorderSources>();
+		UTakeRecorderSources* Sources = ITakeRecorderSourcesManager::GetChecked().FindOrAddSources(InSequence);
 		for (auto Source : Sources->GetSources())
 		{
 			if (Source->IsA<UTakeRecorderActorSource>())
@@ -221,7 +223,7 @@ namespace TakeRecorderSourcesUtils
 			return nullptr;
 		}
 		
-		UTakeRecorderSources* Sources = InLevelSequence->FindOrAddMetaData<UTakeRecorderSources>();
+		UTakeRecorderSources* Sources = ITakeRecorderSourcesManager::GetChecked().FindOrAddSources(InLevelSequence);
 
 		UTakeRecorderActorSource* ActorSource = Sources->AddSource<UTakeRecorderActorSource>();
 		ActorSource->Target = InActor;
diff --git a/Engine/Plugins/VirtualProduction/Takes/Source/TakeRecorderSources/Private/TakeRecorderWorldSource.cpp b/Engine/Plugins/VirtualProduction/Takes/Source/TakeRecorderSources/Private/TakeRecorderWorldSource.cpp
--- a/Engine/Plugins/VirtualProduction/Takes/Source/TakeRecorderSources/Private/TakeRecorderWorldSource.cpp
+++ b/Engine/Plugins/VirtualProduction/Takes/Source/TakeRecorderSources/Private/TakeRecorderWorldSource.cpp
@@ -13,7 +13,7 @@
 
 #include "ISequencer.h"
 #include "MovieSceneFolder.h"
-
+#include "ITakeRecorderSourcesManager.h"
 
 #include "ILevelSequenceEditorToolkit.h"
 #include "Subsystems/AssetEditorSubsystem.h"
@@ -49,7 +49,7 @@ TArray<UTakeRecorderSource*> UTakeRecorderWorldSource::PreRecording(ULevelSequen
 {
 	TArray<UTakeRecorderSource*> NewSources;
 
-	UTakeRecorderSources* Sources = InSequence->FindOrAddMetaData<UTakeRecorderSources>();
+	UTakeRecorderSources* Sources = ITakeRecorderSourcesManager::GetChecked().FindOrAddSources(InSequence);
 
 	// Get the first PIE world's world settings
 	UWorld* World = TakeRecorderSourcesUtils::GetSourceWorld(InSequence);
@@ -169,7 +169,7 @@ void UTakeRecorderWorldSource::AutotrackActors(class ULevelSequence* InSequence,
 		return;
 	}
 
-	UTakeRecorderSources* Sources = InSequence->FindOrAddMetaData<UTakeRecorderSources>();
+	UTakeRecorderSources* Sources = ITakeRecorderSourcesManager::GetChecked().FindOrAddSources(InSequence);
 
 	TArray<AActor*> ActorsBeingRecorded;
 	if (WorldSource.IsValid() && WorldSource.Get()->Target.IsValid())
diff --git a/Engine/Plugins/VirtualProduction/Takes/Source/TakeTrackRecorders/Private/MovieSceneActorReferenceTrackRecorder.cpp b/Engine/Plugins/VirtualProduction/Takes/Source/TakeTrackRecorders/Private/MovieSceneActorReferenceTrackRecorder.cpp
--- a/Engine/Plugins/VirtualProduction/Takes/Source/TakeTrackRecorders/Private/MovieSceneActorReferenceTrackRecorder.cpp
+++ b/Engine/Plugins/VirtualProduction/Takes/Source/TakeTrackRecorders/Private/MovieSceneActorReferenceTrackRecorder.cpp
@@ -13,6 +13,7 @@
 #include "TakesCoreLog.h"
 #include "TakesUtils.h"
 #include "Tracks/MovieSceneActorReferenceTrack.h"
+#include "ITakeRecorderSourcesManager.h"
 
 #include UE_INLINE_GENERATED_CPP_BY_NAME(MovieSceneActorReferenceTrackRecorder)
 
@@ -259,7 +260,7 @@ FMovieSceneActorReferenceKey UMovieSceneActorReferenceTrackRecorder::ActorToRefe
 			FGuid ParentActorId = TakesUtils::ResolveActorFromSequence(InActor, TargetLevelSequence);
 			if (!ParentActorId.IsValid())
 			{
-				UTakeRecorderSources* Sources = TargetLevelSequence->FindOrAddMetaData<UTakeRecorderSources>();
+				UTakeRecorderSources* Sources = ITakeRecorderSourcesManager::GetChecked().FindOrAddSources(TargetLevelSequence);
 				TargetLevelSequence = Sources->AddSourceForActorBinding(InActor, TargetLevelSequence);
 				check(TargetLevelSequence);
 				ParentActorId = TakesUtils::ResolveActorFromSequence(InActor, TargetLevelSequence);
diff --git a/Engine/Plugins/VirtualProduction/Takes/Source/TakesCore/Private/TakesUtils.cpp b/Engine/Plugins/VirtualProduction/Takes/Source/TakesCore/Private/TakesUtils.cpp
--- a/Engine/Plugins/VirtualProduction/Takes/Source/TakesCore/Private/TakesUtils.cpp
+++ b/Engine/Plugins/VirtualProduction/Takes/Source/TakesCore/Private/TakesUtils.cpp
@@ -17,6 +17,7 @@
 #include "TrackRecorders/IMovieSceneTrackRecorderHost.h"
 #include "Tracks/MovieSceneCameraCutTrack.h"
 #include "UObject/SavePackage.h"
+#include "ITakeRecorderSourcesManager.h"
 
 #include "Serialization/ArchiveReplaceOrClearExternalReferences.h"
 
@@ -69,7 +70,7 @@ public:
 		{
 			if (const AActor* Actor = Cast<AActor>(Obj))
 			{
-				if (const UTakeRecorderSources* Sources = LevelSequence->FindMetaData<UTakeRecorderSources>())
+				if (const UTakeRecorderSources* Sources = ITakeRecorderSourcesManager::GetChecked().FindSources(LevelSequence))
 				{
 					return Sources->IsActorPossessable(Actor);
 				}
diff --git a/Engine/Plugins/VirtualProduction/Takes/Source/TakesCore/Public/TakeRecorderSources.h b/Engine/Plugins/VirtualProduction/Takes/Source/TakesCore/Public/TakeRecorderSources.h
--- a/Engine/Plugins/VirtualProduction/Takes/Source/TakesCore/Public/TakeRecorderSources.h
+++ b/Engine/Plugins/VirtualProduction/Takes/Source/TakesCore/Public/TakeRecorderSources.h
@@ -45,7 +45,10 @@ namespace UE::TakeRecorderSources
 	
 
 /**
- * A list of sources to record for any given take. Stored as meta-data on ULevelSequence through ULevelSequence::FindMetaData<UTakeRecorderSources>
+ * A list of sources to record for any given take. 
+ * Is stored and looked up against a level sequence using the ITakeRecorderSourcesManager.
+ * Note that within an editor environment, ITakeRecorderSourcesManager will store the sources object
+ * as meta-data on ULevelSequence through ULevelSequence::FindMetaData<UTakeRecorderSources>.
  */
 UCLASS(MinimalAPI, BlueprintType, Blueprintable)
 class UTakeRecorderSources : public UObject
