Take Recorder - Restructure TakeRecorderNamingTokens code to make module ready for runtime type change.

#rnx
#Synced-CL 52377012

diff --git a/Engine/Plugins/VirtualProduction/Takes/Source/TakeRecorderNamingTokens/Private/TakeRecorderNamingTokens.cpp b/Engine/Plugins/VirtualProduction/Takes/Source/TakeRecorderNamingTokens/Private/TakeRecorderNamingTokens.cpp
--- a/Engine/Plugins/VirtualProduction/Takes/Source/TakeRecorderNamingTokens/Private/TakeRecorderNamingTokens.cpp
+++ b/Engine/Plugins/VirtualProduction/Takes/Source/TakeRecorderNamingTokens/Private/TakeRecorderNamingTokens.cpp
@@ -12,12 +12,15 @@
 #include "TakeRecorderSettings.h"
 #include "TakesCoreBlueprintLibrary.h"
 
-#include "Editor.h"
 #include "Engine/Engine.h"
 #include "LevelSequence.h"
 #include "Misc/PackageName.h"
 #include "Misc/Paths.h"
+
+#if WITH_EDITOR
+#include "Editor.h"
 #include "MovieSceneToolsProjectSettings.h"
+#endif // WITH_EDITOR
 
 #include UE_INLINE_GENERATED_CPP_BY_NAME(TakeRecorderNamingTokens)
 
@@ -89,8 +92,13 @@ void UTakeRecorderNamingTokens::OnCreateDefaultTokens(TArray<FNamingTokenData>&
 	    TEXT("take"),
 	    LOCTEXT("TokenTake", "Take Number"),
 	    FNamingTokenData::FTokenProcessorDelegateNative::CreateLambda([this]{
+		#if WITH_EDITOR
 	        const UMovieSceneToolsProjectSettings* ProjectSettings = GetDefault<UMovieSceneToolsProjectSettings>();
 	        const int32 TakeNumDigits = ProjectSettings->TakeNumDigits;
+		#else // WITH_EDITOR
+			// This loses the settings aspect - to consider if we want to try to bring this to runtime as well.
+			const int32 TakeNumDigits = 2;
+		#endif // WITH_EDITOR
 	        return FText::FromString(FString::Printf(TEXT("%0*d"), TakeNumDigits, TakeMetaData.IsValid() ? TakeMetaData->GetTakeNumber() : 0));
 	    })
 	});
@@ -143,12 +151,14 @@ void UTakeRecorderNamingTokens::OnCreateDefaultTokens(TArray<FNamingTokenData>&
 	    LOCTEXT("TokenMap", "Map"),
 	    FNamingTokenData::FTokenProcessorDelegateNative::CreateLambda([]{
 	        FString MapName;
+		#if WITH_EDITOR
 	        if (GIsEditor)
 	        {
 	            MapName = FPackageName::GetShortFName(
 	                GEditor->GetEditorWorldContext().World()->PersistentLevel->GetOutermost()->GetFName()
 	            ).GetPlainNameString();
 	        }
+		#endif // WITH_EDITOR
 	        return FText::FromString(MapName);
 	    })
 	});
@@ -159,7 +169,12 @@ void UTakeRecorderNamingTokens::OnCreateDefaultTokens(TArray<FNamingTokenData>&
 		FNamingTokenData::FTokenProcessorDelegateNative::CreateLambda([this]{
 			if (Context && Context->Actor.IsValid())
 			{
+			#if WITH_EDITOR
+				// GetActorLabel is an Editor only function.
 				return FText::FromString(Context->Actor->GetActorLabel());
+			#else // WITH_EDITOR
+				return FText::FromString(Context->Actor->GetName());
+			#endif // WITH_EDITOR
 			}
 
 			UE_LOG(LogTakeRecorderNamingTokens, Verbose, TEXT("Attempted to use 'actor' naming token but no context is available."));
diff --git a/Engine/Plugins/VirtualProduction/Takes/Source/TakeRecorderNamingTokens/TakeRecorderNamingTokens.Build.cs b/Engine/Plugins/VirtualProduction/Takes/Source/TakeRecorderNamingTokens/TakeRecorderNamingTokens.Build.cs
--- a/Engine/Plugins/VirtualProduction/Takes/Source/TakeRecorderNamingTokens/TakeRecorderNamingTokens.Build.cs
+++ b/Engine/Plugins/VirtualProduction/Takes/Source/TakeRecorderNamingTokens/TakeRecorderNamingTokens.Build.cs
@@ -10,8 +10,6 @@ public class TakeRecorderNamingTokens : ModuleRules
 			new string[] {
 				"Core",
 				"CoreUObject",
-				"UnrealEd",
-				"MovieSceneTools"
 			}
 		);
 
@@ -23,5 +21,15 @@ public class TakeRecorderNamingTokens : ModuleRules
 				"NamingTokens"
 			}
 		);
+
+		if (Target.bBuildEditor == true)
+		{
+			PrivateDependencyModuleNames.AddRange(
+				new string[] {
+					"UnrealEd",
+					"MovieSceneTools"
+				}
+			);
+		}
 	}
 }
