Take Recorder: Add new CopySources function to ITakeRecorderSourcesManager interface to copy a UTakeRecorderSources object directly.

#rnx
#Synced-CL 52502561

diff --git a/Engine/Plugins/VirtualProduction/Takes/Source/TakesCore/Private/TakeRecorderSourcesManagerImpl.cpp b/Engine/Plugins/VirtualProduction/Takes/Source/TakesCore/Private/TakeRecorderSourcesManagerImpl.cpp
--- a/Engine/Plugins/VirtualProduction/Takes/Source/TakesCore/Private/TakeRecorderSourcesManagerImpl.cpp
+++ b/Engine/Plugins/VirtualProduction/Takes/Source/TakesCore/Private/TakeRecorderSourcesManagerImpl.cpp
@@ -37,6 +37,23 @@ UTakeRecorderSources* UTakeRecorderSourcesManagerImpl::CopySources(ULevelSequenc
 	return nullptr;
 }
 
+UTakeRecorderSources* UTakeRecorderSourcesManagerImpl::CopySources(UTakeRecorderSources* InSources, ULevelSequence* ToSequence)
+{
+	if (!IsValid(ToSequence))
+	{
+		return nullptr;
+	}
+
+#if WITH_EDITOR
+	// when in Editor, ensure they are added to the metadata, overwriting anything thats already there.
+	return ToSequence->CopyMetaData<UTakeRecorderSources>(InSources);
+#else // WITH_EDITOR
+	UTakeRecorderSources* CopiedSources = DuplicateObject(InSources, ToSequence);
+	RuntimeSources.Add(ToSequence, CopiedSources);
+	return CopiedSources;
+#endif // WITH_EDITOR
+}
+
 UTakeRecorderSources* UTakeRecorderSourcesManagerImpl::FindSources(const ULevelSequence* InSequence) const
 {
 	if (!IsValid(InSequence))
diff --git a/Engine/Plugins/VirtualProduction/Takes/Source/TakesCore/Private/TakeRecorderSourcesManagerImpl.h b/Engine/Plugins/VirtualProduction/Takes/Source/TakesCore/Private/TakeRecorderSourcesManagerImpl.h
--- a/Engine/Plugins/VirtualProduction/Takes/Source/TakesCore/Private/TakeRecorderSourcesManagerImpl.h
+++ b/Engine/Plugins/VirtualProduction/Takes/Source/TakesCore/Private/TakeRecorderSourcesManagerImpl.h
@@ -30,6 +30,7 @@ public:
 
 	//~ Begin ITakeRecorderSourcesManager
 	virtual UTakeRecorderSources* CopySources(ULevelSequence* FromSequence, ULevelSequence* ToSequence) override;
+	virtual UTakeRecorderSources* CopySources(UTakeRecorderSources* InSources, ULevelSequence* ToSequence) override;
 	virtual UTakeRecorderSources* FindSources(const ULevelSequence* InSequence) const override;
 	virtual UTakeRecorderSources* FindOrAddSources(ULevelSequence* InSequence) override;
 	virtual bool HasSources(const ULevelSequence* InSequence) const override;
diff --git a/Engine/Plugins/VirtualProduction/Takes/Source/TakesCore/Public/ITakeRecorderSourcesManager.h b/Engine/Plugins/VirtualProduction/Takes/Source/TakesCore/Public/ITakeRecorderSourcesManager.h
--- a/Engine/Plugins/VirtualProduction/Takes/Source/TakesCore/Public/ITakeRecorderSourcesManager.h
+++ b/Engine/Plugins/VirtualProduction/Takes/Source/TakesCore/Public/ITakeRecorderSourcesManager.h
@@ -47,6 +47,15 @@ public:
 	 */
 	virtual UTakeRecorderSources* CopySources(ULevelSequence* FromSequence, ULevelSequence* ToSequence) = 0;
 
+	/**
+	 * Copies the given sources onto the given sequence. This will replace any sources already on ToSequence.
+	 *
+	 * @param InSources     The sources to copy.
+	 * @param ToSequence    The sequence to copy sources to.
+	 * @return Pointer to the newly copied Sources object, or nullptr if nothing was copied.
+	 */
+	virtual UTakeRecorderSources* CopySources(UTakeRecorderSources* InSources, ULevelSequence* ToSequence) = 0;
+
 	/**
 	 * Returns the UTakeRecorderSources associated with the given level sequence, or nullptr if none exist.
 	 *
