Take Recorder - Modify the TakesCore module code to make it runtime compliant and allow a future move to making the TakesCore module a runtime module.

#rnx
#Synced-CL 52673380

diff --git a/Engine/Plugins/VirtualProduction/Takes/Source/TakesCore/Private/TakeData.h b/Engine/Plugins/VirtualProduction/Takes/Source/TakesCore/Private/TakeData.h
--- a/Engine/Plugins/VirtualProduction/Takes/Source/TakesCore/Private/TakeData.h
+++ b/Engine/Plugins/VirtualProduction/Takes/Source/TakesCore/Private/TakeData.h
@@ -2,13 +2,18 @@
 
 #pragma once
 
+// This whole file is only active in an editor context, and inherits from an editor type.
+// It is currently not needed at runtime, but if it is, then this will need revisiting to
+// add that same functionality.
+#if WITH_EDITOR
+
 #include "TakesCoreLog.h"
 #include "TakeMetaData.h"
 #include "TakesCoreBlueprintLibrary.h"
 
 #include "LevelSequence.h"
-#include "MovieSceneToolsModule.h"
 #include "MovieSceneSequence.h"
+#include "MovieSceneToolsModule.h"
 #include "Sections/MovieSceneSubSection.h"
 
 #include "AssetRegistry/AssetRegistryModule.h"
@@ -125,3 +130,5 @@ public:
 		return true;
 	}
 };
+
+#endif // WITH_EDITOR
diff --git a/Engine/Plugins/VirtualProduction/Takes/Source/TakesCore/Private/TakeMetaData.cpp b/Engine/Plugins/VirtualProduction/Takes/Source/TakesCore/Private/TakeMetaData.cpp
--- a/Engine/Plugins/VirtualProduction/Takes/Source/TakesCore/Private/TakeMetaData.cpp
+++ b/Engine/Plugins/VirtualProduction/Takes/Source/TakesCore/Private/TakeMetaData.cpp
@@ -2,7 +2,6 @@
 
 #include "TakeMetaData.h"
 
-#include "Editor.h"
 #include "Engine/Engine.h"
 #include "Engine/Level.h"
 #include "ITakeRecorderNamingTokensModule.h"
@@ -15,6 +14,7 @@
 #include "UObject/AssetRegistryTagsContext.h"
 #include "UObject/Package.h"
 #include "Utils/NamingTokenUtils.h"
+#include "Misc/App.h"
 
 #include UE_INLINE_GENERATED_CPP_BY_NAME(TakeMetaData)
 
@@ -375,6 +375,8 @@ void UTakeMetaData::ExtendAssetRegistryTags(FAssetRegistryTagsContext Context) c
 										FAssetRegistryTag::ETagType::TT_Alphabetical,	FAssetRegistryTag::TD_None));
 }
 
+#if WITH_EDITOR
+
 void UTakeMetaData::ExtendAssetRegistryTagMetaData(TMap<FName, FAssetRegistryTagMetadata>& OutMetadata) const
 {
 	OutMetadata.Add(AssetRegistryTag_Slate, FAssetRegistryTagMetadata()
@@ -413,4 +415,5 @@ void UTakeMetaData::ExtendAssetRegistryTagMetaData(TMap<FName, FAssetRegistryTag
 	);
 }
 
+#endif // WITH_EDITOR
 
diff --git a/Engine/Plugins/VirtualProduction/Takes/Source/TakesCore/Private/TakePreset.cpp b/Engine/Plugins/VirtualProduction/Takes/Source/TakesCore/Private/TakePreset.cpp
--- a/Engine/Plugins/VirtualProduction/Takes/Source/TakesCore/Private/TakePreset.cpp
+++ b/Engine/Plugins/VirtualProduction/Takes/Source/TakesCore/Private/TakePreset.cpp
@@ -3,7 +3,11 @@
 #include "TakePreset.h"
 #include "LevelSequence.h"
 #include "MovieScene.h"
+
+#if WITH_EDITOR
 #include "MovieSceneToolsProjectSettings.h"
+#endif // WITH_EDITOR
+
 #include "TakePresetSettings.h"
 #include "Misc/TransactionObjectEvent.h"
 #include "UObject/Package.h"
@@ -72,19 +76,29 @@ void UTakePreset::CreateLevelSequence()
 	LevelSequence = NewObject<ULevelSequence>(this, UTakePresetSettings::Get()->GetTargetRecordClass(), GetFName(), SequenceFlags);
 	LevelSequence->Initialize();
 
+#if WITH_EDITOR
 	const UMovieSceneToolsProjectSettings* ProjectSettings = GetDefault<UMovieSceneToolsProjectSettings>();
-	UMovieScene* MovieScene = LevelSequence->GetMovieScene();
+	const float DefaultStartTime = ProjectSettings->DefaultStartTime;
+	const float DefaultDuration = ProjectSettings->DefaultDuration;
+#else // WITH_EDITOR
+	// These values are the default values from the Project Settings.
+	const float DefaultStartTime = 0.0f;
+	const float DefaultDuration = 5.0f;
+#endif // WITH_EDITOR
 
-	FFrameNumber StartFrame = (ProjectSettings->DefaultStartTime * MovieScene->GetTickResolution()).RoundToFrame();
-	int32        Duration = (ProjectSettings->DefaultDuration * MovieScene->GetTickResolution()).RoundToFrame().Value;
+	UMovieScene* MovieScene = LevelSequence->GetMovieScene();
+	FFrameNumber StartFrame = (DefaultStartTime * MovieScene->GetTickResolution()).RoundToFrame();
+	int32        Duration = (DefaultDuration * MovieScene->GetTickResolution()).RoundToFrame().Value;
 
 	LevelSequence->GetMovieScene()->SetPlaybackRange(StartFrame, Duration);
 
+#if WITH_EDITOR
 	FMovieSceneEditorData& EditorData = LevelSequence->GetMovieScene()->GetEditorData();
 	EditorData.ViewStart = -1.0;
 	EditorData.ViewEnd   =  5.0;
 	EditorData.WorkStart = -1.0;
 	EditorData.WorkEnd   =  5.0;
+#endif // WITH_EDITOR
 
 	OnLevelSequenceChangedEvent.Broadcast();
 }
@@ -144,6 +158,7 @@ void UTakePreset::RemoveOnLevelSequenceChanged(FDelegateHandle DelegateHandle)
 	OnLevelSequenceChangedEvent.Remove(DelegateHandle);
 }
 
+#if WITH_EDITOR
 void UTakePreset::PostTransacted(const FTransactionObjectEvent& TransactionEvent)
 {
 	Super::PostTransacted(TransactionEvent);
@@ -153,4 +168,5 @@ void UTakePreset::PostTransacted(const FTransactionObjectEvent& TransactionEvent
 		OnLevelSequenceChangedEvent.Broadcast();
 	}
 }
+#endif // WITH_EDITOR
 
diff --git a/Engine/Plugins/VirtualProduction/Takes/Source/TakesCore/Private/TakeRecorderSourceProperty.cpp b/Engine/Plugins/VirtualProduction/Takes/Source/TakesCore/Private/TakeRecorderSourceProperty.cpp
--- a/Engine/Plugins/VirtualProduction/Takes/Source/TakesCore/Private/TakeRecorderSourceProperty.cpp
+++ b/Engine/Plugins/VirtualProduction/Takes/Source/TakesCore/Private/TakeRecorderSourceProperty.cpp
@@ -5,6 +5,7 @@
 
 #include UE_INLINE_GENERATED_CPP_BY_NAME(TakeRecorderSourceProperty)
 
+#if WITH_EDITOR
 void UActorRecorderPropertyMap::PostEditUndo()
 {
 	Super::PostEditUndo();
@@ -16,6 +17,7 @@ void UActorRecorderPropertyMap::PostEditChangeProperty(struct FPropertyChangedEv
 	Super::PostEditChangeProperty(PropertyChangedEvent);
 	UpdateCachedValues();
 }
+#endif // WITH_EDITOR
 
 void UActorRecorderPropertyMap::UpdateCachedValues()
 {
diff --git a/Engine/Plugins/VirtualProduction/Takes/Source/TakesCore/Private/TakeRecorderSources.cpp b/Engine/Plugins/VirtualProduction/Takes/Source/TakesCore/Private/TakeRecorderSources.cpp
--- a/Engine/Plugins/VirtualProduction/Takes/Source/TakesCore/Private/TakeRecorderSources.cpp
+++ b/Engine/Plugins/VirtualProduction/Takes/Source/TakesCore/Private/TakeRecorderSources.cpp
@@ -4,7 +4,6 @@
 #include "AssetRegistry/AssetData.h"
 #include "AssetRegistry/AssetRegistryModule.h"
 #include "CanvasTypes.h"
-#include "Editor.h"
 #include "Engine/Engine.h"
 #include "Engine/Font.h"
 #include "Engine/TimecodeProvider.h"
@@ -15,7 +14,11 @@
 #include "MovieScene.h"
 #include "MovieSceneFolder.h"
 #include "NamingTokens/TakeRecorderNamingTokensContext.h"
-#include "ObjectTools.h"
+#include "Engine/TimecodeProvider.h"
+#include "Misc/App.h"
+#include "MetaData/MovieSceneShotMetaData.h"
+#include "NamingTokens/TakeRecorderNamingTokensContext.h"
+#include "UObject/Package.h"
 #include "Sections/MovieSceneSubSection.h"
 #include "TakeMetaData.h"
 #include "TakeRecorderSource.h"
@@ -24,6 +27,11 @@
 #include "TrackRecorders/IMovieSceneTrackRecorderHost.h"
 #include "Tracks/MovieSceneSubTrack.h"
 
+#if WITH_EDITOR
+#include "ObjectTools.h"
+#include "Editor.h"
+#endif // WITH_EDITOR
+
 #define LOCTEXT_NAMESPACE "TakeRecorderSources"
 
 #include UE_INLINE_GENERATED_CPP_BY_NAME(TakeRecorderSources)
@@ -162,6 +170,7 @@ bool UTakeRecorderSources::IsActorPossessable(const AActor* InActor) const
 	return GetSettings().bRecordToPossessable;
 }
 
+#if WITH_EDITOR
 void UTakeRecorderSources::PostEditChangeProperty(FPropertyChangedEvent& PropertyChangedEvent)
 {
 	Super::PostEditChangeProperty(PropertyChangedEvent);
@@ -183,6 +192,7 @@ void UTakeRecorderSources::PostEditChangeProperty(FPropertyChangedEvent& Propert
 		SlotIndex++;
 	}
 }
+#endif // WITH_EDITOR
 
 void UTakeRecorderSources::StartPreRecordingRecursive(const TArray<UTakeRecorderSource*>& InSources, ULevelSequence* InRootSequence,
 	TArray<UTakeRecorderSource*>& NewSourcesOut, FManifestSerializer* InManifestSerializer)
@@ -228,8 +238,7 @@ void UTakeRecorderSources::PreRecordingRecursive(TArray<UTakeRecorderSource*> In
 			// is using which sub-sequence so that we can push the correct sequence for all points of the Source's recording lifecycle.
 			if (Settings.bRecordSourcesIntoSubSequences && Source->SupportsSubscenes())
 			{
-				const FString SubSequenceAssetName = ObjectTools::SanitizeObjectName(Source->GetSubsceneAssetName(InRootSequence));
-				const FString BaseSubSequenceTrackName = ObjectTools::SanitizeObjectName(Source->GetSubsceneTrackName(InRootSequence));
+				const FString BaseSubSequenceTrackName = TakesUtils::SanitizeObjectName(Source->GetSubsceneTrackName(InRootSequence));
 				FString SubSequenceTrackName = BaseSubSequenceTrackName;
 				
 				// Simple duplicate name handling to prevent collisions. Normally the user should have a unique token in the title,
@@ -247,7 +256,8 @@ void UTakeRecorderSources::PreRecordingRecursive(TArray<UTakeRecorderSource*> In
 					}
 					while (bTrackNameAlreadyUsed);
 				}
-				
+
+				const FString SubSequenceAssetName = TakesUtils::SanitizeObjectName(Source->GetSubsceneAssetName(InRootSequence));
 				TargetSequence = CreateSubSequenceForSource(InRootSequence, SubSequenceTrackName, SubSequenceAssetName, Settings);
 
 				// If there's already a Subscene Track for our sub-sequence we need to remove that track before create a new one. No data is lost in this process as the
@@ -255,6 +265,13 @@ void UTakeRecorderSources::PreRecordingRecursive(TArray<UTakeRecorderSource*> In
 				const FString SequenceName = FPaths::GetBaseFilename(TargetSequence->GetPathName());
 				UMovieSceneSubTrack* SubsceneTrack = nullptr;
 
+				/*
+				 * Note that as part of the TakeRecorder in runtime work and GetDisplayName() being an editor only function, the logic below is currently only run in editor.
+				 * When adding runtime support, with a non editor alternative to GetDisplayName, the below for loop in the if branch will need updating.
+				 * As Take Recorder at runtime is not yet support, it is ok for this not to function or work at runtime, but will eventually be updated.
+				 */
+
+			#if WITH_EDITOR
 				if (Source->OverwriteExistingTrackData())
 				{
 					for (UMovieSceneTrack* Track : InRootSequence->GetMovieScene()->GetTracks())
@@ -274,6 +291,7 @@ void UTakeRecorderSources::PreRecordingRecursive(TArray<UTakeRecorderSource*> In
 						}
 					}
 				}
+			#endif // WITH_EDITOR
 
 				// We need to add the new subsequence to the root sequence immediately so that it shows up in the UI and you can tell that things
 				// are being recorded, otherwise they don't show up until recording stops and then it magically pops in.
@@ -285,9 +303,11 @@ void UTakeRecorderSources::PreRecordingRecursive(TArray<UTakeRecorderSource*> In
 				// Track should not be transactional during the recording process
 				SubsceneTrack->ClearFlags(RF_Transactional);
 
+			#if WITH_EDITOR
 				// We create a new sub track for every Source so that we can name the Subtrack after the Source instead of just the sections within it.
 				SubsceneTrack->SetDisplayName(FText::FromString(Source->GetSubsceneTrackName(InRootSequence)));
 				SubsceneTrack->SetColorTint(Source->TrackTint);
+			#endif // WITH_EDITOR
 
 				// When we create the Subscene Track we'll make sure a folder is created for it to sort into and add the new Subscene Track as a child of it.
 				if (bCreateSequencerFolders)
@@ -606,14 +626,18 @@ void UTakeRecorderSources::StopRecording(class ULevelSequence* InSequence, const
 				TakesUtils::ClampPlaybackRangeToEncompassAllSections(SubSequence->GetMovieScene(), bUpperBoundOnly);
 				TakesUtils::ResetViewAndWorkRange(SubSequence->GetMovieScene());
 
+			#if WITH_EDITOR
+				// SetReadOnly is editor only.
 				// Lock the sequence so that it can't be changed without implicitly unlocking it now
 				if (Settings.bAutoLock)
 				{
 					SubSequence->GetMovieScene()->SetReadOnly(true);
 				}
+			#endif // WITH_EDITOR
 
 				ULevelSequence* SequenceAsset = CastChecked<ULevelSequence>(SubSequence);
 
+			#if WITH_EDITOR
 				// Mark that this subsequence has been recorded and is a subsequence.
 				UMovieSceneShotMetaData* MovieSequenceMetaData = SequenceAsset->FindOrAddMetaData<UMovieSceneShotMetaData>();
 				MovieSequenceMetaData->SetIsRecorded(true);
@@ -623,6 +647,7 @@ void UTakeRecorderSources::StopRecording(class ULevelSequence* InSequence, const
 				UTakeMetaData* TakeMetaData = SequenceAsset->FindMetaData<UTakeMetaData>();
 				check(TakeMetaData);
 				TakeMetaData->Lock();
+			#endif // WITH_EDITOR
 
 				SubSection->SetRange(SubSequence->GetMovieScene()->GetPlaybackRange());
 
@@ -683,10 +708,12 @@ void UTakeRecorderSources::StopRecording(class ULevelSequence* InSequence, const
 	CachedManifestSerializer = nullptr;
 	CachedLevelSequence = nullptr;
 
+#if WITH_EDITOR
 	if (GEditor && AssetsToCleanUp.Num() > 0)
 	{
 		ObjectTools::ForceDeleteObjects(AssetsToCleanUp, false);
 	}
+#endif // WITH_EDITOR
 }
 
 ULevelSequence* UTakeRecorderSources::FindExistingSubSequence(const ULevelSequence* InRootSequence, const FString& InSubSequenceAssetName)
@@ -727,6 +754,8 @@ ULevelSequence* UTakeRecorderSources::CreateSubSequenceForSource(ULevelSequence*
 	}
 	
 	FString SubSequenceDirectory;
+
+#if WITH_EDITOR
 	if (const UTakeMetaData* AssetMetaData = InRootSequence->FindMetaData<UTakeMetaData>())
 	{
 		// User specified subsequence directory.
@@ -735,6 +764,7 @@ ULevelSequence* UTakeRecorderSources::CreateSubSequenceForSource(ULevelSequence*
 		SubSequenceDirectory = InSettings.SubSequenceDirectory;
 		SubSequenceDirectory = AssetMetaData->GenerateAssetPath(SubSequenceDirectory, Context);
 	}
+#endif // WITH_EDITOR
 	
 	if (SubSequenceDirectory.IsEmpty())
 	{
@@ -762,14 +792,15 @@ ULevelSequence* UTakeRecorderSources::CreateSubSequenceForSource(ULevelSequence*
 			OutAsset->GetMovieScene()->SetDisplayRate(InRootSequence->GetMovieScene()->GetDisplayRate());
 		}
 
-		UTakeMetaData* TakeMetaData = InRootSequence->FindMetaData<UTakeMetaData>();
-		if (TakeMetaData)
+#if WITH_EDITOR
+		if (UTakeMetaData* TakeMetaData = InRootSequence->FindMetaData<UTakeMetaData>())
 		{
 			UTakeMetaData* OutTakeMetaData = OutAsset->CopyMetaData(TakeMetaData);
 
 			// Tack on the sub sequence name so that it's unique from the root sequence
 			OutTakeMetaData->SetSlate(TakeMetaData->GetSlate() + TEXT("_") + SubSequenceTrackName, false);
 		}
+#endif // WITH_EDITOR
 
 		OutAsset->MarkPackageDirty();
 		FAssetRegistryModule::AssetCreated(OutAsset);
@@ -788,11 +819,14 @@ UMovieSceneFolder* UTakeRecorderSources::AddFolderForSource(const UTakeRecorderS
 	// the Take Recorder UI, which is currently not customizable. If that becomes customizable this code should be updated
 	// to ensure the created folder structure matches the one visible in the Take Recorder UI.
 
-	// Currently we use the category that the Source is filed under as this is what the UI currently sorts by.
-	const FName FolderName = FName(*InSource->GetClass()->GetMetaData(FName("Category")));
-
 	// Search the Movie Scene for a folder with this name
 	UMovieSceneFolder* FolderToUse = nullptr;
+
+	FName FolderName = *InSource->GetName();
+
+#if WITH_EDITOR
+	// Currently we use the category that the Source is filed under as this is what the UI currently sorts by.
+	FolderName = FName(*InSource->GetClass()->GetMetaData(FName("Category")));
 	for (UMovieSceneFolder* Folder : InMovieScene->GetRootFolders())
 	{
 		if (Folder->GetFolderName() == FolderName)
@@ -801,17 +835,22 @@ UMovieSceneFolder* UTakeRecorderSources::AddFolderForSource(const UTakeRecorderS
 			break;
 		}
 	}
+#endif // WITH_EDITOR
 
 	// If we didn't find a folder with this name we're going to go ahead and create a new folder
 	if (FolderToUse == nullptr)
 	{
 		FolderToUse = NewObject<UMovieSceneFolder>(InMovieScene, NAME_None, RF_Transactional);
 		FolderToUse->SetFolderName(FolderName);
+	#if WITH_EDITOR
 		InMovieScene->AddRootFolder(FolderToUse);
+	#endif // WITH_EDITOR
 	}
 
+#if WITH_EDITOR
 	// We want to expand these folders in the Sequencer UI (since these are visible as they record).
 	InMovieScene->GetEditorData().ExpansionStates.FindOrAdd(FolderName.ToString()) = FMovieSceneExpansionState(true);
+#endif // WITH_EDITOR
 
 	return FolderToUse;
 }
diff --git a/Engine/Plugins/VirtualProduction/Takes/Source/TakesCore/Private/TakeRecorderTimeProcessing.cpp b/Engine/Plugins/VirtualProduction/Takes/Source/TakesCore/Private/TakeRecorderTimeProcessing.cpp
--- a/Engine/Plugins/VirtualProduction/Takes/Source/TakesCore/Private/TakeRecorderTimeProcessing.cpp
+++ b/Engine/Plugins/VirtualProduction/Takes/Source/TakesCore/Private/TakeRecorderTimeProcessing.cpp
@@ -1,4 +1,4 @@
-﻿// Copyright Epic Games, Inc. All Rights Reserved.
+// Copyright Epic Games, Inc. All Rights Reserved.
 
 #include "TakeRecorderTimeProcessing.h"
 
@@ -63,10 +63,12 @@ void ProcessRecordedTimes(ULevelSequence* InSequence, UMovieSceneTakeTrack* Take
 		TakeSection->TimecodeSource = EarliestTimecodeSource;
 	}
 
+#if WITH_EDITORONLY_DATA
 	if (UTakeMetaData* TakeMetaData = InSequence->FindMetaData<UTakeMetaData>())
 	{
 		TakeSection->Slate.SetDefault(FString::Printf(TEXT("%s_%d"), *TakeMetaData->GetSlate(), TakeMetaData->GetTakeNumber()));
 	}
+#endif // WITH_EDITORONLY_DATA
 
 	if (TakeSection->GetAutoSizeRange().IsSet())
 	{
diff --git a/Engine/Plugins/VirtualProduction/Takes/Source/TakesCore/Private/TakesCoreModule.cpp b/Engine/Plugins/VirtualProduction/Takes/Source/TakesCore/Private/TakesCoreModule.cpp
--- a/Engine/Plugins/VirtualProduction/Takes/Source/TakesCore/Private/TakesCoreModule.cpp
+++ b/Engine/Plugins/VirtualProduction/Takes/Source/TakesCore/Private/TakesCoreModule.cpp
@@ -2,12 +2,15 @@
 
 #include "TakesCoreLog.h"
 #include "TakeMetaData.h"
-#include "TakeData.h"
 
 #include "LevelSequence.h"
-#include "MovieSceneToolsModule.h"
 
-#include "Modules/ModuleManager.h"
+#include "Modules/ModuleInterface.h"
+
+#if WITH_EDITOR
+#include "TakeData.h"
+#include "MovieSceneToolsModule.h"
+#endif // WITH_EDITOR
 
 LLM_DEFINE_TAG(Takes_TakesCore);
 
@@ -20,35 +23,42 @@ class FTakesCoreModule : public IModuleInterface
 		// Ensure the level sequence module is loaded
 		FModuleManager::Get().LoadModuleChecked("LevelSequence");
 
+		LevelSequenceCDO = GetMutableDefault<ULevelSequence>();
+
+	
+	#if WITH_EDITOR
+		TakeData = MakeShareable(new FTakesCoreTakeData);
 		// Add empty take meta data to the ULevelSequence CDO to ensure that
 		// asset registry tooltips show up in the editor
 		UTakeMetaData* MetaData = GetMutableDefault<ULevelSequence>()->FindOrAddMetaData<UTakeMetaData>();
 		MetaData->SetFlags(RF_Transient);
 
-		LevelSequenceCDO = GetMutableDefault<ULevelSequence>();
-
 		// Register take data with movie scene tools so sequencer knows how to switch takes
-		TakeData = MakeShareable(new FTakesCoreTakeData);
 		FMovieSceneToolsModule::Get().RegisterTakeData(TakeData.Get());
+	#endif // WITH_EDITOR
 	}
 
 	virtual void ShutdownModule() override
 	{
 		LLM_SCOPE_BYTAG(Takes_TakesCore);
 
+	#if WITH_EDITOR
 		if (ULevelSequence* CDO = LevelSequenceCDO.Get())
 		{
 			CDO->RemoveMetaData<UTakeMetaData>();
 		}
 
 		FMovieSceneToolsModule::Get().UnregisterTakeData(TakeData.Get());
+	#endif // WITH_EDITOR
 	}
 
 	// Weak ptr to the level sequence CDO so we can gracefully remove the take meta-data on shutdown module
 	// without crashing when ShutdownModule is called after the CDO has been destroyed.
 	TWeakObjectPtr<ULevelSequence> LevelSequenceCDO;
 
+#if WITH_EDITORONLY_DATA
 	TSharedPtr<IMovieSceneToolsTakeData> TakeData;
+#endif // WITH_EDITORONLY_DATA
 };
 
 IMPLEMENT_MODULE(FTakesCoreModule, TakesCore);
diff --git a/Engine/Plugins/VirtualProduction/Takes/Source/TakesCore/Private/TakesUtils.cpp b/Engine/Plugins/VirtualProduction/Takes/Source/TakesCore/Private/TakesUtils.cpp
--- a/Engine/Plugins/VirtualProduction/Takes/Source/TakesCore/Private/TakesUtils.cpp
+++ b/Engine/Plugins/VirtualProduction/Takes/Source/TakesCore/Private/TakesUtils.cpp
@@ -2,14 +2,11 @@
 
 #include "TakesUtils.h"
 #include "AssetRegistry/AssetData.h"
-#include "TakesCoreLog.h"
 
 #include "Engine/Engine.h"
 #include "Engine/World.h"
-#include "ILevelSequenceEditorToolkit.h"
 #include "LevelSequence.h"
 #include "Math/Range.h"
-#include "MovieScene.h"
 #include "MovieSceneSection.h"
 #include "MovieSceneTimeHelpers.h"
 #include "Sections/MovieSceneCameraCutSection.h"
@@ -21,6 +18,10 @@
 
 #include "Serialization/ArchiveReplaceOrClearExternalReferences.h"
 
+#if WITH_EDITOR
+#include "ILevelSequenceEditorToolkit.h"
+#endif // WITH_EDITOR
+
 #define LOCTEXT_NAMESPACE "TakesUtils"
 
 namespace TakesUtils
@@ -98,6 +99,54 @@ static void LogRemovedExternalObjects(UObject* ReferencingObject, const TArray<U
 }
 
 
+FString SanitizeInvalidChars(const FString& InText, const TCHAR* InvalidChars)
+{
+	FString SanitizedText = InText;
+
+	const TCHAR* InvalidChar = InvalidChars ? InvalidChars : TEXT("");
+	while (*InvalidChar)
+	{
+		SanitizedText.ReplaceCharInline(*InvalidChar, TCHAR('_'), ESearchCase::CaseSensitive);
+		++InvalidChar;
+	}
+
+	return SanitizedText;
+}
+
+FString SanitizeObjectName(const FString& InObjectName)
+{
+	return SanitizeInvalidChars(InObjectName, INVALID_OBJECTNAME_CHARACTERS);
+}
+
+FString SanitizePackageName(const FString& InPackageName)
+{
+	FString SanitizedName = SanitizeInvalidChars(InPackageName, INVALID_LONGPACKAGE_CHARACTERS);
+
+	// Coalesce multiple contiguous slashes into a single slash
+	int32 CharIndex = 0;
+	while (CharIndex < SanitizedName.Len())
+	{
+		if (SanitizedName[CharIndex] == TEXT('/'))
+		{
+			int32 SlashCount = 1;
+			while (CharIndex + SlashCount < SanitizedName.Len() &&
+				SanitizedName[CharIndex + SlashCount] == TEXT('/'))
+			{
+				SlashCount++;
+			}
+
+			if (SlashCount > 1)
+			{
+				SanitizedName.RemoveAt(CharIndex + 1, SlashCount - 1, EAllowShrinking::No);
+			}
+		}
+
+		CharIndex++;
+	}
+
+	return SanitizedName;
+}
+
 /** Helper function - get the first PIE world (or first PIE client world if there is more than one) */
 UWorld* GetFirstPIEWorld()
 {
@@ -162,10 +211,12 @@ void ResetViewAndWorkRange(UMovieScene* InMovieScene)
 	const double OutputViewSize = PlayRange.Size<FFrameNumber>() / TickResolution;
 	const double OutputChange = OutputViewSize * 0.1;
 
+#if WITH_EDITOR
 	TRange<double> NewRange = UE::MovieScene::ExpandRange(PlayRange / TickResolution, OutputChange);
 	FMovieSceneEditorData& EditorData = InMovieScene->GetEditorData();
 	EditorData.ViewStart = EditorData.WorkStart = NewRange.GetLowerBoundValue();
 	EditorData.ViewEnd = EditorData.WorkEnd = NewRange.GetUpperBoundValue();
+#endif // WITH_EDITOR
 }
 
 void CleanupExternalReferences(UObject* InObject, UPackage* InPackage)
@@ -290,6 +341,7 @@ UWorld* DiscoverSourceWorld()
 	return SourceWorld;
 }
 
+#if WITH_EDITOR
 TSharedPtr<ISequencer> OpenSequencer(ULevelSequence* LevelSequence, FText* OutError)
 {
 	TSharedPtr<ISequencer> Sequencer;
@@ -317,12 +369,21 @@ FQualifiedFrameTime GetRecordTime(TSharedPtr<ISequencer> Sequencer, ULevelSequen
 {
 	FQualifiedFrameTime RecordTime;
 
-	if (Sequencer.IsValid())
-	{
-		// Sequencer handles time dilation.
-		RecordTime = Sequencer->GetGlobalTime();
-	}
-	else if (SequenceAsset)
+	// Sequencer handles time dilation.
+	RecordTime = Sequencer.IsValid()
+		? Sequencer->GetGlobalTime()
+		: GetRecordTime(SequenceAsset, TimecodeAtStart, bStartAtCurrentTimecode, TimeDilation);
+
+	return RecordTime;
+}
+#endif // WITH_EDITOR
+
+FQualifiedFrameTime GetRecordTime(ULevelSequence* SequenceAsset,
+	const FTimecode& TimecodeAtStart, bool bStartAtCurrentTimecode, float TimeDilation)
+{
+	FQualifiedFrameTime RecordTime;
+
+	if (SequenceAsset)
 	{
 		if (UMovieScene* MovieScene = SequenceAsset->GetMovieScene())
 		{
@@ -359,7 +420,14 @@ FGuid ResolveActorFromSequence(const AActor* InActor, const ULevelSequence* Curr
 	for (int32 PossessableCount = 0; PossessableCount < MovieScene->GetPossessableCount(); ++PossessableCount)
 	{
 		const FMovieScenePossessable& Possessable = MovieScene->GetPossessable(PossessableCount);
+	#if WITH_EDITOR
+		// GetActorLabel is an editor only function, so for runtime, use just the name directly.
+		// When we have a better way of identifying actors without their labels, this will need updating.
 		if (Possessable.GetName() == InActor->GetActorLabel())
+	#else // WITH_EDITOR
+		if (Possessable.GetName() == InActor->GetName())
+	#endif // WITH_EDITOR
+
 		{
 			return Possessable.GetGuid();
 		}
diff --git a/Engine/Plugins/VirtualProduction/Takes/Source/TakesCore/Public/TakeMetaData.h b/Engine/Plugins/VirtualProduction/Takes/Source/TakesCore/Public/TakeMetaData.h
--- a/Engine/Plugins/VirtualProduction/Takes/Source/TakesCore/Public/TakeMetaData.h
+++ b/Engine/Plugins/VirtualProduction/Takes/Source/TakesCore/Public/TakeMetaData.h
@@ -82,10 +82,12 @@ public:
 	UE_DEPRECATED(5.4, "Implement the version that takes FAssetRegistryTagsContext instead.")
 	virtual void ExtendAssetRegistryTags(TArray<FAssetRegistryTag>& OutTags) const override {}
 
+#if WITH_EDITOR
 	/**
 	 * Extend the default ULevelSequence asset registry tag meta-data
 	 */
 	UE_API virtual void ExtendAssetRegistryTagMetaData(TMap<FName, FAssetRegistryTagMetadata>& OutMetadata) const override;
+#endif // WITH_EDITOR
 
 public:
 
diff --git a/Engine/Plugins/VirtualProduction/Takes/Source/TakesCore/Public/TakePreset.h b/Engine/Plugins/VirtualProduction/Takes/Source/TakesCore/Public/TakePreset.h
--- a/Engine/Plugins/VirtualProduction/Takes/Source/TakesCore/Public/TakePreset.h
+++ b/Engine/Plugins/VirtualProduction/Takes/Source/TakesCore/Public/TakePreset.h
@@ -65,7 +65,9 @@ public:
 	static UE_API UTakePreset* AllocateTransientPreset(UTakePreset* TemplatePreset);
 private:
 
+#if WITH_EDITOR
 	UE_API virtual void PostTransacted(const FTransactionObjectEvent& TransactionEvent) override;
+#endif // WITH_EDITOR
 
 	/** Instanced level sequence template that is used to define a starting point for a new take recording */
 	UPROPERTY(Instanced)
diff --git a/Engine/Plugins/VirtualProduction/Takes/Source/TakesCore/Public/TakeRecorderSourceProperty.h b/Engine/Plugins/VirtualProduction/Takes/Source/TakesCore/Public/TakeRecorderSourceProperty.h
--- a/Engine/Plugins/VirtualProduction/Takes/Source/TakesCore/Public/TakeRecorderSourceProperty.h
+++ b/Engine/Plugins/VirtualProduction/Takes/Source/TakesCore/Public/TakeRecorderSourceProperty.h
@@ -49,8 +49,10 @@ class UActorRecorderPropertyMap : public UObject
 	GENERATED_BODY()
 
 public:
+#if WITH_EDITOR
 	UE_API virtual void PostEditUndo() override;
 	UE_API virtual void PostEditChangeProperty(struct FPropertyChangedEvent& PropertyChangedEvent) override;
+#endif // WITH_EDITOR
 
 	UPROPERTY(BlueprintReadOnly, VisibleAnywhere, Category = "Property")
 	TSoftObjectPtr<UObject> RecordedObject;
diff --git a/Engine/Plugins/VirtualProduction/Takes/Source/TakesCore/Public/TakeRecorderSources.h b/Engine/Plugins/VirtualProduction/Takes/Source/TakesCore/Public/TakeRecorderSources.h
--- a/Engine/Plugins/VirtualProduction/Takes/Source/TakesCore/Public/TakeRecorderSources.h
+++ b/Engine/Plugins/VirtualProduction/Takes/Source/TakesCore/Public/TakeRecorderSources.h
@@ -238,8 +238,10 @@ private:
 	/** Called at the end of each frame in both the Editor and in Game to update all Sources. */
 	virtual void Tick(float DeltaTime) {}
 
+#if WITH_EDITOR
 	/** Called whenever a property is changed on this class */
 	UE_API virtual void PostEditChangeProperty(FPropertyChangedEvent& PropertyChangedEvent) override;
+#endif // WITH_EDITOR
 
 	/** A list of handlers to invoke when the sources list changes */
 	DECLARE_MULTICAST_DELEGATE(FOnSourcesChanged);
diff --git a/Engine/Plugins/VirtualProduction/Takes/Source/TakesCore/Public/TakesUtils.h b/Engine/Plugins/VirtualProduction/Takes/Source/TakesCore/Public/TakesUtils.h
--- a/Engine/Plugins/VirtualProduction/Takes/Source/TakesCore/Public/TakesUtils.h
+++ b/Engine/Plugins/VirtualProduction/Takes/Source/TakesCore/Public/TakesUtils.h
@@ -4,7 +4,6 @@
 
 #include "AssetRegistry/AssetData.h"
 #include "CoreMinimal.h"
-#include "ISequencer.h"
 #include "Misc/PackageName.h"
 #include "MovieScene.h"
 #include "Modules/ModuleManager.h"
@@ -13,7 +12,10 @@
 #include "UObject/UObjectGlobals.h"
 #include "TakesCoreLog.h"
 #include "Misc/Paths.h"
-#include "PackageTools.h"
+
+#if WITH_EDITOR
+#include "ISequencer.h"
+#endif // WITH_EDITOR
 
 class UTakeRecorderSource;
 class UWorld;
@@ -21,6 +23,30 @@ class ULevelSequence;
 
 namespace TakesUtils
 {
+	/**
+	 * Converts all invalid object name characters to _.
+	 *
+	 * Taken from ObjectTools and moved a small subset here to remove the need to call into ObjectTools and allow
+	 * the same functionality in a non-Editor environment.
+	 */
+	TAKESCORE_API FString SanitizeInvalidChars(const FString& InText, const TCHAR* InvalidChars);
+
+	/**
+	 * Converts all invalid object name characters to _.
+	 *
+	 * Taken from ObjectTools and moved a small subset here to remove the need to call into ObjectTools and allow
+	 * the same functionality in a non-Editor environment.
+	 */
+	TAKESCORE_API FString SanitizeObjectName(const FString& InObjectName);
+
+	/**
+	 * Replaces all invalid package name characters with _.
+	 *
+	 * Taken from PackageTools and moved a small subset here to remove the need to call into PackageTools and allow
+	 * the same functionality in a non-Editor environment.
+	 */
+	TAKESCORE_API FString SanitizePackageName(const FString& InPackageName);
+
 	/*
 	 * Get the first PIE world (or first PIE client world if there is more than one)
 	 */
@@ -139,8 +165,8 @@ namespace TakesUtils
 				return FSoftObjectPath((BaseAssetPath / AssetName) + TEXT(".") + AssetName);
 			};
 
-		FString SanitizedBaseAssetPath = UPackageTools::SanitizePackageName(BaseAssetPath);
-		FString SanitizedBaseAssetName = UPackageTools::SanitizePackageName(BaseAssetName);
+		FString SanitizedBaseAssetPath = SanitizePackageName(BaseAssetPath);
+		FString SanitizedBaseAssetName = SanitizePackageName(BaseAssetName);
 
 		FString SourceAssetName = FPaths::MakeValidFileName(SanitizedBaseAssetName);
 		FSoftObjectPath AssetObjectPath = MakeAssetObjectPath(SanitizedBaseAssetPath, SourceAssetName);
@@ -177,15 +203,21 @@ namespace TakesUtils
 
 	TAKESCORE_API UWorld* DiscoverSourceWorld();
 
+#if WITH_EDITOR
 	TAKESCORE_API TSharedPtr<ISequencer> OpenSequencer(ULevelSequence* LevelSequence, FText* OutError);
 	
 	/* The time at which to record. Taken from the Sequencer global time, otherwise based on timecode */
 	TAKESCORE_API FQualifiedFrameTime GetRecordTime(TSharedPtr<ISequencer> Sequencer, ULevelSequence* SequenceAsset,
 		const FTimecode& TimecodeAtStart, bool bStartAtCurrentTimecode, float TimeDilation = 1.f);
+#endif // WITH_EDITOR
+
+	/* The time at which to record, based on timecode */
+	TAKESCORE_API FQualifiedFrameTime GetRecordTime(ULevelSequence* SequenceAsset,
+		const FTimecode& TimecodeAtStart, bool bStartAtCurrentTimecode, float TimeDilation = 1.f);
 	
 	/** Retrieve an actor guid for moviescene. */
 	TAKESCORE_API FGuid ResolveActorFromSequence(const AActor* InActor, const ULevelSequence* CurrentSequence);
 	
 	/** Checks if a source is considered possessable. */
 	TAKESCORE_API bool IsSourcePossessable(const UTakeRecorderSource* InSource);
-}
+} // namespace TakesUtils
diff --git a/Engine/Plugins/VirtualProduction/Takes/Source/TakesCore/TakesCore.Build.cs b/Engine/Plugins/VirtualProduction/Takes/Source/TakesCore/TakesCore.Build.cs
--- a/Engine/Plugins/VirtualProduction/Takes/Source/TakesCore/TakesCore.Build.cs
+++ b/Engine/Plugins/VirtualProduction/Takes/Source/TakesCore/TakesCore.Build.cs
@@ -10,8 +10,6 @@ public class TakesCore : ModuleRules
 			new string[] {
 				"Core",
 				"CoreUObject",
-				"EditorFramework",
-				"UnrealEd",
 				"NamingTokens", 
 				"TakeMovieScene",
 			}
@@ -24,20 +22,34 @@ public class TakesCore : ModuleRules
 				"SlateCore",
 				"MovieScene",
 				"MovieSceneTracks",
-				"MovieSceneTools",
-				"LevelSequence",
-				"LevelSequenceEditor",
 				"Engine",
-				"SerializedRecorderInterface",
-
+				"SerializedRecorderInterface"
 			}
 		);
 		
 		PublicIncludePathModuleNames.AddRange(
 			new string[]
 			{
-				"TakeRecorderNamingTokens"
+				"TakeRecorderNamingTokens",
+				"TakeTrackRecorders"
 			}
 		);
+
+		if (Target.bBuildEditor == true)
+		{
+			PrivateDependencyModuleNames.AddRange(
+				new string[] {
+					"EditorFramework",
+					"UnrealEd",
+				}
+			);
+
+			PublicDependencyModuleNames.AddRange(
+				new string[] {
+					"MovieSceneTools",
+					"LevelSequenceEditor",
+				}
+			);
+		}
 	}
 }
diff --git a/Engine/Source/Runtime/MovieScene/Public/MovieSceneSection.h b/Engine/Source/Runtime/MovieScene/Public/MovieSceneSection.h
--- a/Engine/Source/Runtime/MovieScene/Public/MovieSceneSection.h
+++ b/Engine/Source/Runtime/MovieScene/Public/MovieSceneSection.h
@@ -798,11 +798,9 @@ public:
 	UPROPERTY(EditAnywhere, Category="Section")
 	FMovieSceneFrameRange SectionRange;
 
-#if WITH_EDITORONLY_DATA
 	/** The timecode at which this movie scene section is based (ie. when it was recorded) */
 	UPROPERTY(EditAnywhere, Category="Section")
 	FMovieSceneTimecodeSource TimecodeSource;
-#endif
 
 	/** Optional dynamic condition for whether this section evaluates at runtime. */
 	UPROPERTY(EditAnywhere, Category="Section")
