Take Recorder - Restructure and update TakeRecorderSources module to be ready for a change to a runtime module.

#rnx
#Synced-CL 52674206

diff --git a/Engine/Plugins/VirtualProduction/Takes/Source/TakeRecorderSources/Private/ActorTakeRecorderDropHandler.cpp b/Engine/Plugins/VirtualProduction/Takes/Source/TakeRecorderSources/Private/ActorTakeRecorderDropHandler.cpp
--- a/Engine/Plugins/VirtualProduction/Takes/Source/TakeRecorderSources/Private/ActorTakeRecorderDropHandler.cpp
+++ b/Engine/Plugins/VirtualProduction/Takes/Source/TakeRecorderSources/Private/ActorTakeRecorderDropHandler.cpp
@@ -2,10 +2,13 @@
 
 #include "ActorTakeRecorderDropHandler.h"
 
+#if WITH_EDITOR
 #include "DragAndDrop/ActorDragDropOp.h"
 #include "DragAndDrop/CompositeDragDropOp.h"
 #include "DragAndDrop/FolderDragDropOp.h"
 #include "EditorActorFolders.h"
+#endif // WITH_EDITOR
+
 #include "TakeRecorderSourceHelpers.h"
 #include "TakeRecorderSourcesUtils.h"
 
@@ -19,6 +22,7 @@ void FActorTakeRecorderDropHandler::HandleOperation(TSharedPtr<FDragDropOperatio
 
 bool FActorTakeRecorderDropHandler::CanHandleOperation(TSharedPtr<FDragDropOperation> InOperation, UTakeRecorderSources* Sources)
 {
+#if WITH_EDITOR
 	bool bCanHandle = false;
 	if (InOperation)
 	{		
@@ -78,10 +82,14 @@ bool FActorTakeRecorderDropHandler::CanHandleOperation(TSharedPtr<FDragDropOpera
 	}
 
 	return bCanHandle;
+#else // WITH_EDITOR
+	return false;
+#endif // WITH_EDITOR
 }
 
 TArray<AActor*> FActorTakeRecorderDropHandler::GetValidDropActors(TSharedPtr<FDragDropOperation> InOperation, UTakeRecorderSources* Sources)
 {
+#if WITH_EDITOR
 	TSharedPtr<FActorDragDropOp>  ActorDrag = nullptr;
 	TSharedPtr<FFolderDragDropOp> FolderDrag = nullptr;
 
@@ -178,5 +186,8 @@ TArray<AActor*> FActorTakeRecorderDropHandler::GetValidDropActors(TSharedPtr<FDr
 	}
 
 	return DraggedActors;
+#else // WITH_EDITOR
+	return TArray<AActor*>();
+#endif // WITH_EDITOR
 }
 }
\ No newline at end of file
diff --git a/Engine/Plugins/VirtualProduction/Takes/Source/TakeRecorderSources/Private/TakeRecorderActorSource.cpp b/Engine/Plugins/VirtualProduction/Takes/Source/TakeRecorderSources/Private/TakeRecorderActorSource.cpp
--- a/Engine/Plugins/VirtualProduction/Takes/Source/TakeRecorderSources/Private/TakeRecorderActorSource.cpp
+++ b/Engine/Plugins/VirtualProduction/Takes/Source/TakeRecorderSources/Private/TakeRecorderActorSource.cpp
@@ -1,18 +1,12 @@
 // Copyright Epic Games, Inc. All Rights Reserved.
 
 #include "TakeRecorderActorSource.h"
-#include "AnimationRecorder.h"
 #include "Styling/SlateIconFinder.h"
-#include "ClassIconFinder.h"
 #include "MovieScene.h"
 #include "LevelSequence.h"
-#include "Editor.h"
 #include "Features/IModularFeatures.h"
-#include "ActorRecordingSettings.h"
 #include "Misc/ScopedSlowTask.h"
-#include "SequenceRecorderUtils.h"
 #include "TakeRecorderSource.h"
-#include "TakeRecorderSourceHelpers.h"
 #include "TakeRecorderSources.h"
 #include "TakeRecorderSourcesUtils.h"
 #include "Recorder/TakeRecorderParameters.h"
@@ -21,15 +15,11 @@
 #include "TakeMetaData.h"
 #include "MovieSceneFolder.h"
 #include "Serializers/MovieSceneManifestSerialization.h"
-#include "Compilation/MovieSceneCompiledDataManager.h"
 
-#include "Animation/SkeletalMeshActor.h"
 #include "GameFramework/Actor.h"
+#include "GameFramework/Pawn.h"
 #include "GameFramework/MovementComponent.h"
-#include "GameFramework/Character.h"
 #include "Camera/CameraComponent.h"
-#include "CameraRig_Crane.h"
-#include "CameraRig_Rail.h"
 #include "Components/SkeletalMeshComponent.h"
 
 #include "TrackRecorders/IMovieSceneTrackRecorderFactory.h"
@@ -38,22 +28,21 @@
 #include "TrackRecorders/MovieSceneAnimationTrackRecorder.h"
 #include "TrackRecorders/MovieSceneAnimationTrackRecorderSettings.h"
 #include "TrackRecorders/MovieScene3DTransformTrackRecorder.h"
-#include "TrackRecorders/MovieScenePropertyTrackRecorder.h"
 #include "TrackRecorders/MovieSceneTrackRecorderSettings.h"
 
 #include "MovieSceneTakeTrack.h"
-#include "MovieSceneTakeSection.h"
 #include "MovieSceneTakeSettings.h"
 
-#if WITH_EDITOR
-#include "Tracks/MovieSceneSpawnTrack.h"
-#endif
 #include "NamingTokens/TakeRecorderNamingTokensContext.h"
 #include "Recorder/TakeRecorderSubsystem.h"
-#include "SequencerUtilities.h"
 #include "TakeRecorderSourcesModule.h"
 #include "TakeRecorderTimeProcessing.h"
 
+#if WITH_EDITOR
+#include "AnimationRecorder.h"
+#include "SequencerUtilities.h"
+#endif // WITH_EDITOR
+
 #include UE_INLINE_GENERATED_CPP_BY_NAME(TakeRecorderActorSource)
 
 DEFINE_LOG_CATEGORY(ActorSerialization);
@@ -89,6 +78,7 @@ UTakeRecorderSource* UTakeRecorderActorSource::AddSourceForActor(AActor* InActor
 		return nullptr;
 	}
 
+#if WITH_EDITOR
 	//Look through our sources and see if one actor matches the incoming one either from editor or PIE world.
 	{
 		//Cache  InputActor comparison data
@@ -134,6 +124,7 @@ UTakeRecorderSource* UTakeRecorderActorSource::AddSourceForActor(AActor* InActor
 			}
 		}
 	}
+#endif // WITH_EDITOR
 
 	UTakeRecorderActorSource* NewSource = InSources->AddSource<UTakeRecorderActorSource>();
 	NewSource->SetSourceActor(InActor);
@@ -154,6 +145,7 @@ void UTakeRecorderActorSource::RemoveActorFromSources(AActor* InActor, UTakeReco
 		return;
 	}
 
+#if WITH_EDITOR
 	//Cache  InputActor comparison data
 	const bool bIsAlreadyPIEActor = InActor->GetOutermost()->HasAnyPackageFlags(PKG_PlayInEditor);
 	const AActor* InputActorEditor = EditorUtilities::GetEditorWorldCounterpartActor(InActor);
@@ -196,6 +188,7 @@ void UTakeRecorderActorSource::RemoveActorFromSources(AActor* InActor, UTakeReco
 			}
 		}
 	}
+#endif // WITH_EDITOR
 }
 
 namespace TakeRecorderActorSource
@@ -259,14 +252,24 @@ TArray<UTakeRecorderSource*> UTakeRecorderActorSource::PreRecording(ULevelSequen
 	}
 	
 	// Initialize the header for this actor in the Manifest Serializer for streaming data capture.
+#if WITH_EDITOR
+	const FString ActorId = ActorToRecord->GetActorLabel();
+#else // WITH_EDITOR
+	const FString ActorId = ActorToRecord->GetName();
+#endif // WITH_EDITOR
+
 	FName SerializedType("Actor");
-	FActorFileHeader Header(ObjectBindingName, ActorToRecord->GetActorLabel(), SerializedType, ActorToRecord->GetClass()->GetPathName(), false);
+	FActorFileHeader Header(ObjectBindingName, ActorId, SerializedType, ActorToRecord->GetClass()->GetPathName(), false);
 
 	if (GetRecordToPossessable())
 	{
 		// If a user adds a PIE-only instance as a recordable object, they can't record this to a possessable (because the binding will be broken when they exit PIE).
 		if (!CVarTakeRecorderAllowPossessablePIEObjects->GetBool() &&
-			Target->GetWorld()->WorldType != EWorldType::Editor && GEditor && !GEditor->ObjectsThatExistInEditorWorld.Get(Target.Get()))
+			Target->GetWorld()->WorldType != EWorldType::Editor 
+		#if WITH_EDITOR
+			&& GEditor && !GEditor->ObjectsThatExistInEditorWorld.Get(Target.Get())
+		#endif // WITH_EDITOR
+			)
 		{
 			UE_LOGF(LogTakesCore, Warning, "Attempted to record an actor that does not exist in the editor world to a possessable. Forcing recording of %ls as a Spawnable so that the resulting binding is not broken!", *Target->GetName());
 			RecordType = ETakeRecorderActorRecordType::Spawnable;
@@ -274,7 +277,7 @@ TArray<UTakeRecorderSource*> UTakeRecorderActorSource::PreRecording(ULevelSequen
 		else
 		{
 			// Create a Possessable object binding in the Sequence and then bind it to our actor
-			CachedObjectBindingGuid = MovieScene->AddPossessable(ActorToRecord->GetActorLabel(), ActorToRecord->GetClass());
+			CachedObjectBindingGuid = MovieScene->AddPossessable(ActorId, ActorToRecord->GetClass());
 			TargetLevelSequence->BindPossessableObject(CachedObjectBindingGuid, *ActorToRecord, ActorToRecord->GetWorld());
 			Header.bRecordToPossessable = true;
 		}
@@ -282,6 +285,7 @@ TArray<UTakeRecorderSource*> UTakeRecorderActorSource::PreRecording(ULevelSequen
 	
 	if (!GetRecordToPossessable())
 	{
+	#if WITH_EDITOR
 		// We need to store the object template in the Movie Scene (because it's a complex UObject)
 		// instead of trying to place this data into the non-UObject safe data stream.
 		UE::Sequencer::FCreateBindingParams CreateBindingParams;
@@ -289,6 +293,10 @@ TArray<UTakeRecorderSource*> UTakeRecorderActorSource::PreRecording(ULevelSequen
 		CreateBindingParams.bAllowCustomBinding = true;
 		CreateBindingParams.BindingNameOverride = ActorToRecord->GetActorLabel();
 		CachedObjectBindingGuid = FSequencerUtilities::CreateOrReplaceBinding(nullptr, InSequence, ActorToRecord, CreateBindingParams);
+	#else // WITH_EDITOR
+		CachedObjectBindingGuid = MovieSceneHelpers::TryCreateCustomSpawnableBinding(InSequence, ActorToRecord);
+	#endif // WITH_EDITOR
+
 		CachedObjectTemplate = Cast<AActor>(MovieSceneHelpers::GetObjectTemplate(InSequence, CachedObjectBindingGuid, MovieSceneHelpers::CreateTransientSharedPlaybackState(ActorToRecord, InSequence)));
 		
 		if (CachedObjectTemplate.IsValid())
@@ -444,7 +452,11 @@ void UTakeRecorderActorSource::CreateSectionRecordersRecursive(UObject* ObjectTo
 		for (int32 PossessableCount = 0; PossessableCount < MovieScene->GetPossessableCount(); ++PossessableCount)
 		{
 			const FMovieScenePossessable& Possessable = MovieScene->GetPossessable(PossessableCount);
-			if (Possessable.GetParent() == CachedObjectBindingGuid && Possessable.GetName() == Component->GetName() && Possessable.GetPossessedObjectClass() == Component->GetClass())
+			if (Possessable.GetParent() == CachedObjectBindingGuid && Possessable.GetName() == Component->GetName()
+			#if WITH_EDITOR
+				&& Possessable.GetPossessedObjectClass() == Component->GetClass()
+			#endif // WITH_EDITOR
+				)
 			{
 				Guid = Possessable.GetGuid();
 				break;
@@ -840,7 +852,18 @@ TArray<UTakeRecorderSource*> UTakeRecorderActorSource::PostRecording(ULevelSeque
 	Parameters.User = GetDefault<UTakeRecorderUserSettings>()->Settings;
 	Parameters.Project = GetDefault<UTakeRecorderProjectSettings>()->Settings;
 
-	FScopedSlowTask SlowTask((float)TrackRecorders.Num() + 1.0f, FText::Format(LOCTEXT("ProcessingActor", "Generating MovieScene Data for Actor {0}"), Target.IsValid() ? FText::FromString(Target.Get()->GetActorLabel()) : FText::GetEmpty()));
+	FString TargetId;
+
+	if (Target.IsValid())
+	{
+	#if WITH_EDITOR
+		TargetId = Target.Get()->GetActorLabel();
+	#else // WITH_EDITOR
+		TargetId = Target.Get()->GetName();
+	#endif // WITH_EDITOR
+	}
+
+	FScopedSlowTask SlowTask((float)TrackRecorders.Num() + 1.0f, FText::Format(LOCTEXT("ProcessingActor", "Generating MovieScene Data for Actor {0}"), !TargetId.IsEmpty() ? FText::FromString(TargetId) : FText::GetEmpty()));
 	SlowTask.MakeDialog(false, bShowProgressDialog);
 
 	// We need to do some post-processing tasks on the Track Recorders (such as animation motion source fixup) so we do this now before finalizing
@@ -893,11 +916,13 @@ TArray<UTakeRecorderSource*> UTakeRecorderActorSource::PostRecording(ULevelSeque
 				TakesUtils::CreateCameraCutTrack(InRootSequence, RecordedCameraGuid, RecordedCameraSequenceID, InRootSequence->GetMovieScene()->GetPlaybackRange());
 			}
 
+		#if WITH_EDITOR
 			// Swap our target actor to the Editor actor (in case the recording was added while in PIE)
 			if (AActor* EditorActor = EditorUtilities::GetEditorWorldCounterpartActor(Target.Get()))
 			{
 				Target = EditorActor;
 			}
+		#endif // WITH_EDITOR
 		}
 	}
 
@@ -933,6 +958,7 @@ void UTakeRecorderActorSource::FinalizeRecording()
 
 namespace UE::TakeRecorderActorSource::Private
 {
+#if WITH_EDITOR
 FProcessRecordedTimeParams GetTimecodeRecordingParameters()
 {
 	FString HoursName = GetDefault<UMovieSceneTakeSettings>()->HoursName;
@@ -953,6 +979,7 @@ FProcessRecordedTimeParams GetTimecodeRecordingParameters()
 		.SlateName = SlateName
 	};
 }
+#endif // WITH_EDITOR
 }
 void UTakeRecorderActorSource::PostProcessTrackRecorders(ULevelSequence* InSequence)
 {
@@ -960,11 +987,13 @@ void UTakeRecorderActorSource::PostProcessTrackRecorders(ULevelSequence* InSeque
 	Parameters.User = GetDefault<UTakeRecorderUserSettings>()->Settings;
 	Parameters.Project = GetDefault<UTakeRecorderProjectSettings>()->Settings;
 
+#if WITH_EDITOR
 	FProcessRecordedTimeParams RecordedTimeParams = UE::TakeRecorderActorSource::Private::GetTimecodeRecordingParameters();
 	if (UTakeMetaData* TakeMetaData = InSequence->FindMetaData<UTakeMetaData>())
 	{
 		RecordedTimeParams.Slate = FString::Printf(TEXT("%s_%d"), *TakeMetaData->GetSlate(), TakeMetaData->GetTakeNumber());
 	}
+#endif // WITH_EDITOR
 
 	// We want to look at all Animation Track recorders and remove root motion if the transform
 	// for that component is being recorded. We copy the animation out of the Animation Track
@@ -980,7 +1009,11 @@ void UTakeRecorderActorSource::PostProcessTrackRecorders(ULevelSequence* InSeque
 			continue;
 		}
 		AActor* SourceActor = Cast<AActor>(TrackRecorder->GetSourceObject());
-		AActor* SourceEditorActor = SourceActor ? EditorUtilities::GetEditorWorldCounterpartActor(SourceActor) : nullptr;
+		AActor* SourceEditorActor = 
+		#if WITH_EDITOR
+			SourceActor ? EditorUtilities::GetEditorWorldCounterpartActor(SourceActor) : 
+		#endif // WITH_EDITOR
+			nullptr;
 		AActor* TargetActor = Target.Get();
 
 		if (!RootTransformRecorder && TrackRecorder->IsA<UMovieScene3DTransformTrackRecorder>() && (TargetActor == SourceActor || TargetActor == SourceEditorActor))
@@ -1025,7 +1058,9 @@ void UTakeRecorderActorSource::PostProcessTrackRecorders(ULevelSequence* InSeque
 			
 			if (Parameters.Project.bRecordTimecode)
 			{
+			#if WITH_EDITOR
 				AnimationTrackRecorder->ProcessRecordedTimes(RecordedTimeParams);
+			#endif // WITH_EDITOR
 			}
 		}
 	}
@@ -1178,6 +1213,7 @@ bool UTakeRecorderActorSource::EnsureObjectTemplateHasComponent(UActorComponent*
 	return true;
 }
 
+#if WITH_EDITOR
 void UTakeRecorderActorSource::PostEditChangeProperty(FPropertyChangedEvent& PropertyChangedEvent)
 {
 	Super::PostEditChangeProperty(PropertyChangedEvent);
@@ -1212,6 +1248,7 @@ void UTakeRecorderActorSource::PostEditUndo()
 	Super::PostEditUndo();
 	UpdateCachedNumberOfRecordedProperties();
 }
+#endif // WITH_EDITOR
 
 void UTakeRecorderActorSource::RebuildRecordedPropertyMap()
 {
@@ -1275,7 +1312,12 @@ void UTakeRecorderActorSource::RebuildRecordedPropertyMapRecursive(const FFieldV
 	for (TFieldIterator<FProperty> It(ObjectClass); It; ++It)
 	{
 		const bool bIsInterpField = It->HasAllPropertyFlags(CPF_Interp);
+
+	#if WITH_METADATA
 		const bool bHasTrackMetadata = It->HasMetaData(SequencerTrackClassMetadataName);
+	#else // WITH_METADATA
+		const bool bHasTrackMetadata = false;
+	#endif // WITH_METADATA
 
 		FString PropertyName = It->GetFName().ToString();
 		FString PropertyPath = OuterStructPath + FString(*PropertyName);
@@ -1429,7 +1471,11 @@ FText UTakeRecorderActorSource::GetDisplayTextImpl() const
 	AActor* TargetActor = Target.Get();
 	if (TargetActor)
 	{
+	#if WITH_EDITOR
 		return FText::FromString(TargetActor->GetActorLabel());
+	#else // WITH_EDITOR
+		return FText::FromString(TargetActor->GetName());
+	#endif // WITH_EDITOR
 	}
 
 	return LOCTEXT("ActorLabel", "Actor (None)");
@@ -1616,7 +1662,12 @@ void UTakeRecorderActorSource::GetChildSceneComponents(USceneComponent* OnSceneC
 			// If the Scene Component's owner is spawned at Runtime, we record it.
 			if (SceneComponent->GetOwner() != Target.Get())
 			{
-				bool bActorIsTemporary = GEditor && (SceneComponent->GetOwner()->GetWorld()->WorldType == EWorldType::PIE && !GEditor->ObjectsThatExistInEditorWorld.Get(SceneComponent->GetOwner()));
+			#if WITH_EDITOR
+				const bool bActorIsTemporary = GEditor && (SceneComponent->GetOwner()->GetWorld()->WorldType == EWorldType::PIE && !GEditor->ObjectsThatExistInEditorWorld.Get(SceneComponent->GetOwner()));
+			#else // WITH_EDITOR
+				const bool bActorIsTemporary = false;
+			#endif // WITH_EDITOR
+				
 				if (bActorIsTemporary)
 				{
 					if (bUpdateReferencedActorList)
@@ -1774,15 +1825,11 @@ void UTakeRecorderActorSource::CreateNewActorSourceForReferencedActors()
 		NewSources.Add(ActorSource);
 		AddedActorSources.Add(ActorSource);
 
-		ActorSource->Target = Actor;
+		// Setting the source actor through the function ensures internal property maps are also updated.
+		ActorSource->SetSourceActor(Actor);
 
 		// For consistency in the hierarchy, actor sources should have the same state as the source automatically adding them
 		ActorSource->RecordType = RecordType;
-
-		// Send a PropertyChangedEvent so the class catches the callback and rebuilds the property map. We can't rely on the Actor rebuilding the map on PreRecording
-		// because that would wipe out any user adjustments from one added natively.
-		FPropertyChangedEvent PropertyChangedEvent(UTakeRecorderActorSource::StaticClass()->FindPropertyByName(GET_MEMBER_NAME_CHECKED(UTakeRecorderActorSource, Target)), EPropertyChangeType::ValueSet);
-		ActorSource->PostEditChangeProperty(PropertyChangedEvent);
 	}
 	NewReferencedActors.Reset();
 }
@@ -1826,10 +1873,16 @@ FTrackRecorderSettings UTakeRecorderActorSource::GetTrackRecorderSettings() cons
 	Parameters.User = GetDefault<UTakeRecorderUserSettings>()->Settings;
 	Parameters.Project = GetDefault<UTakeRecorderProjectSettings>()->Settings;
 
+#if WITH_EDITOR
+	const bool bEditorIsAvailable = GEditor != nullptr;
+#else // WITH_EDITOR
+	const bool bEditorIsAvailable = false;
+#endif // WITH_EDITOR
+
 	TrackRecorderSettings.bRecordToPossessable = GetRecordToPossessable();
 	TrackRecorderSettings.bReduceKeys = bReduceKeys;
 	TrackRecorderSettings.bRemoveRedundantTracks = Parameters.User.bRemoveRedundantTracks;
-	TrackRecorderSettings.bSaveRecordedAssets = Sources->GetSettings().bSaveRecordedAssets || GEditor == nullptr;
+	TrackRecorderSettings.bSaveRecordedAssets = Sources->GetSettings().bSaveRecordedAssets || !bEditorIsAvailable;
 	TrackRecorderSettings.ReduceKeysTolerance = Parameters.User.ReduceKeysTolerance;
 
 	TrackRecorderSettings.DefaultTracks = Parameters.Project.DefaultTracks;
@@ -1966,6 +2019,7 @@ UMovieSceneTrackRecorderSettings* UTakeRecorderActorSource::GetSettingsObjectFor
 
 FString UTakeRecorderActorSource::GetSubsceneTrackName(ULevelSequence* InSequence) const
 {
+#if WITH_EDITOR
 	if (const UTakeMetaData* TakeMetaData = InSequence->FindMetaData<UTakeMetaData>())
 	{
 		if (const UMovieSceneAnimationTrackRecorderEditorSettings* Settings = GetBestTrackRecorderEditorSettings())
@@ -1975,10 +2029,15 @@ FString UTakeRecorderActorSource::GetSubsceneTrackName(ULevelSequence* InSequenc
 			return TakeMetaData->GenerateAssetPath(Settings->AnimationTrackName.ToString(), Context);
 		}
 	}
+#endif // WITH_EDITOR
 	
 	if (Target.IsValid())
 	{
+	#if WITH_EDITOR
 		return Target->GetActorLabel();
+	#else // WITH_EDITOR
+		return Target->GetName();
+	#endif // WITH_EDITOR
 	}
 
 	return Super::GetSubsceneTrackName(InSequence);
@@ -1986,6 +2045,7 @@ FString UTakeRecorderActorSource::GetSubsceneTrackName(ULevelSequence* InSequenc
 
 FString UTakeRecorderActorSource::GetSubsceneAssetName(ULevelSequence* InSequence) const
 {
+#if WITH_EDITOR
 	if (const UTakeMetaData* TakeMetaData = InSequence->FindMetaData<UTakeMetaData>())
 	{
 		if (const UMovieSceneAnimationTrackRecorderEditorSettings* Settings = GetBestTrackRecorderEditorSettings())
@@ -1995,10 +2055,15 @@ FString UTakeRecorderActorSource::GetSubsceneAssetName(ULevelSequence* InSequenc
 			return TakeMetaData->GenerateAssetPath(Settings->AnimationAssetName, Context);
 		}
 	}
+#endif // WITH_EDITOR
 	
 	if (Target.IsValid())
 	{
+	#if WITH_EDITOR
 		return Target->GetActorLabel();
+	#else // WITH_EDITOR
+		return Target->GetName();
+	#endif // WITH_EDITOR
 	}
 
 	return Super::GetSubsceneAssetName(InSequence);
diff --git a/Engine/Plugins/VirtualProduction/Takes/Source/TakeRecorderSources/Private/TakeRecorderCameraCutSource.cpp b/Engine/Plugins/VirtualProduction/Takes/Source/TakeRecorderSources/Private/TakeRecorderCameraCutSource.cpp
--- a/Engine/Plugins/VirtualProduction/Takes/Source/TakeRecorderSources/Private/TakeRecorderCameraCutSource.cpp
+++ b/Engine/Plugins/VirtualProduction/Takes/Source/TakeRecorderSources/Private/TakeRecorderCameraCutSource.cpp
@@ -2,16 +2,13 @@
 
 #include "TakeRecorderCameraCutSource.h"
 
-#include "AssetRegistry/AssetData.h"
 #include "TakeRecorderSources.h"
 #include "TakeRecorderSource.h"
 #include "TakeRecorderActorSource.h"
 #include "TakeRecorderSourcesUtils.h"
 #include "TakesUtils.h"
-#include "TakesCoreLog.h"
 
 #include "LevelSequence.h"
-#include "LevelEditorViewport.h"
 #include "Engine/Engine.h"
 #include "Camera/PlayerCameraManager.h"
 #include "GameFramework/PlayerController.h"
@@ -19,6 +16,10 @@
 #include "Tracks/MovieSceneCameraCutTrack.h"
 #include "Sections/MovieSceneCameraCutSection.h"
 
+#if WITH_EDITOR
+#include "LevelEditorViewport.h"
+#endif // WITH_EDITOR
+
 #include UE_INLINE_GENERATED_CPP_BY_NAME(TakeRecorderCameraCutSource)
 
 #define LOCTEXT_NAMESPACE "UTakeRecorderCameraCutSource"
@@ -45,6 +46,7 @@ void UTakeRecorderCameraCutSource::TickRecording(const FQualifiedFrameTime& Curr
 		return;
 	}
 
+#if WITH_EDITOR
 	if (World->WorldType == EWorldType::Editor)
 	{
 		if (GCurrentLevelEditingViewportClient)
@@ -56,7 +58,8 @@ void UTakeRecorderCameraCutSource::TickRecording(const FQualifiedFrameTime& Curr
 			}
 		}
 	}
-	else 
+	else
+#endif // WITH_EDITOR
 	{
 		APlayerController* PC = World->GetGameInstance() ? World->GetGameInstance()->GetFirstLocalPlayerController() : nullptr;
 		APlayerCameraManager* PlayerCameraManager = PC ? ToRawPtr(PC->PlayerCameraManager) : nullptr;
@@ -76,12 +79,8 @@ void UTakeRecorderCameraCutSource::TickRecording(const FQualifiedFrameTime& Curr
 		UTakeRecorderSources* OwningSources = CastChecked<UTakeRecorderSources>(GetOuter());
 
 		UTakeRecorderActorSource* ActorSource = OwningSources->AddSource<UTakeRecorderActorSource>();
-		ActorSource->Target = Target;
-
-		// Send a PropertyChangedEvent so the class catches the callback and rebuilds the property map. We can't rely on the Actor rebuilding the map on PreRecording
-		// because that would wipe out any user adjustments from one added natively.
-		FPropertyChangedEvent PropertyChangedEvent(UTakeRecorderActorSource::StaticClass()->FindPropertyByName(GET_MEMBER_NAME_CHECKED(UTakeRecorderActorSource, Target)), EPropertyChangeType::ValueSet);
-		ActorSource->PostEditChangeProperty(PropertyChangedEvent);
+		// Setting the source actor through the function ensures internal property maps are also updated.
+		ActorSource->SetSourceActor(Target);
 
 		// This has to be called after setting the Target and propagating the change event so that it has a chance to know what to record
 		// about the actor.
diff --git a/Engine/Plugins/VirtualProduction/Takes/Source/TakeRecorderSources/Private/TakeRecorderLevelSequenceSource.cpp b/Engine/Plugins/VirtualProduction/Takes/Source/TakeRecorderSources/Private/TakeRecorderLevelSequenceSource.cpp
--- a/Engine/Plugins/VirtualProduction/Takes/Source/TakeRecorderSources/Private/TakeRecorderLevelSequenceSource.cpp
+++ b/Engine/Plugins/VirtualProduction/Takes/Source/TakeRecorderSources/Private/TakeRecorderLevelSequenceSource.cpp
@@ -1,15 +1,11 @@
 // Copyright Epic Games, Inc. All Rights Reserved.
 
 #include "TakeRecorderLevelSequenceSource.h"
-#include "AssetRegistry/AssetData.h"
 #include "TakesCoreLog.h"
 #include "TakesUtils.h"
-#include "Editor.h"
 #include "GameFramework/Actor.h"
 #include "LevelSequenceActor.h"
 #include "LevelSequencePlayer.h"
-#include "Styling/SlateIconFinder.h"
-#include "ClassIconFinder.h"
 
 #include UE_INLINE_GENERATED_CPP_BY_NAME(TakeRecorderLevelSequenceSource)
 
@@ -103,7 +99,11 @@ FText UTakeRecorderLevelSequenceSource::GetDisplayTextImpl() const
 {
 	if (ActorsToTrigger.Num() == 1 && ActorsToTrigger[0].IsValid())
 	{
+	#if WITH_EDITOR
 		return FText::FromString(ActorsToTrigger[0].Get()->GetActorLabel());
+	#else // WITH_EDITOR
+		return FText::FromString(ActorsToTrigger[0].Get()->GetName());
+	#endif // WITH_EDITOR
 	}
 
 	uint32 NumValid = 0;
diff --git a/Engine/Plugins/VirtualProduction/Takes/Source/TakeRecorderSources/Private/TakeRecorderLevelVisibilitySource.cpp b/Engine/Plugins/VirtualProduction/Takes/Source/TakeRecorderSources/Private/TakeRecorderLevelVisibilitySource.cpp
--- a/Engine/Plugins/VirtualProduction/Takes/Source/TakeRecorderSources/Private/TakeRecorderLevelVisibilitySource.cpp
+++ b/Engine/Plugins/VirtualProduction/Takes/Source/TakeRecorderSources/Private/TakeRecorderLevelVisibilitySource.cpp
@@ -1,17 +1,14 @@
 // Copyright Epic Games, Inc. All Rights Reserved.
 
 #include "TakeRecorderLevelVisibilitySource.h"
-#include "AssetRegistry/AssetData.h"
 #include "TakesUtils.h"
 
 #include "TakeRecorderSources.h"
-#include "TakeRecorderActorSource.h"
 #include "TakeRecorderSourcesUtils.h"
 #include "TakeMetaData.h"
 
 #include "LevelSequence.h"
 #include "Engine/Engine.h"
-#include "Modules/ModuleManager.h"
 
 #include "MovieSceneFolder.h"
 #include "Sections/MovieSceneLevelVisibilitySection.h"
@@ -28,6 +25,7 @@ UTakeRecorderLevelVisibilitySourceSettings::UTakeRecorderLevelVisibilitySourceSe
 	TrackTint = FColor(176, 117, 19);
 }
 
+#if WITH_EDITOR
 void UTakeRecorderLevelVisibilitySourceSettings::PostEditChangeProperty(FPropertyChangedEvent& PropertyChangedEvent)
 {
 	Super::PostEditChangeProperty(PropertyChangedEvent);
@@ -37,22 +35,27 @@ void UTakeRecorderLevelVisibilitySourceSettings::PostEditChangeProperty(FPropert
 		SaveConfig();
 	}
 }
+#endif // WITH_EDITOR
 
 FString UTakeRecorderLevelVisibilitySourceSettings::GetSubsceneTrackName(ULevelSequence* InSequence) const
 {
+#if WITH_EDITOR
 	if (UTakeMetaData* TakeMetaData = InSequence->FindMetaData<UTakeMetaData>())
 	{
 		return FString::Printf(TEXT("Level Visibility_%s"), *TakeMetaData->GenerateAssetPath("{slate}"));
 	}
+#endif // WITH_EDITOR
 	return TEXT("Level Visibility");
 }
 
 FString UTakeRecorderLevelVisibilitySourceSettings::GetSubsceneAssetName(ULevelSequence* InSequence) const
 {
+#if WITH_EDITOR
 	if (UTakeMetaData* TakeMetaData = InSequence->FindMetaData<UTakeMetaData>())
 	{
 		return FString::Printf(TEXT("Level Visibility_%s"), *TakeMetaData->GenerateAssetPath("{slate}_{take}"));
 	}
+#endif // WITH_EDITOR
 	return TEXT("Level Visibility");
 }
 
@@ -94,7 +97,11 @@ TArray<UTakeRecorderSource*> UTakeRecorderLevelVisibilitySource::PreRecording(UL
 	UMovieScene* MovieScene = InSequence->GetMovieScene();
 	for (auto Track : MovieScene->GetTracks())
 	{
-		if (Track->IsA(UMovieSceneLevelVisibilityTrack::StaticClass()) && Track->GetDisplayName().EqualTo(LevelVisibilityTrackName))
+		if (Track->IsA(UMovieSceneLevelVisibilityTrack::StaticClass())
+		#if WITH_EDITOR
+			&& Track->GetDisplayName().EqualTo(LevelVisibilityTrackName)
+		#endif // WITH_EDITOR
+			)
 		{
 			CachedLevelVisibilityTrack = Cast<UMovieSceneLevelVisibilityTrack>(Track);
 		}
@@ -103,7 +110,9 @@ TArray<UTakeRecorderSource*> UTakeRecorderLevelVisibilitySource::PreRecording(UL
 	if (!CachedLevelVisibilityTrack.IsValid())
 	{
 		CachedLevelVisibilityTrack = MovieScene->AddTrack<UMovieSceneLevelVisibilityTrack>();
+	#if WITH_EDITOR
 		CachedLevelVisibilityTrack->SetDisplayName(LevelVisibilityTrackName);
+	#endif // WITH_EDITOR
 	}
 
 	CachedLevelVisibilityTrack->RemoveAllAnimationData();
diff --git a/Engine/Plugins/VirtualProduction/Takes/Source/TakeRecorderSources/Private/TakeRecorderMicrophoneAudioManager.cpp b/Engine/Plugins/VirtualProduction/Takes/Source/TakeRecorderSources/Private/TakeRecorderMicrophoneAudioManager.cpp
--- a/Engine/Plugins/VirtualProduction/Takes/Source/TakeRecorderSources/Private/TakeRecorderMicrophoneAudioManager.cpp
+++ b/Engine/Plugins/VirtualProduction/Takes/Source/TakeRecorderSources/Private/TakeRecorderMicrophoneAudioManager.cpp
@@ -2,12 +2,14 @@
 
 #include "TakeRecorderMicrophoneAudioManager.h"
 
-#include "TakeRecorderSettings.h"
-#include "TakeRecorderSources.h"
-#include "TakesUtils.h"
+#include "TakesCoreLog.h"
+#include "Modules/ModuleManager.h"
 
+#if WITH_EDITOR
 #include "IAudioCaptureEditor.h"
 #include "IAudioCaptureEditorModule.h"
+#include "AudioCaptureEditorTypes.h"
+#endif // WITH_EDITOR
 
 #include UE_INLINE_GENERATED_CPP_BY_NAME(TakeRecorderMicrophoneAudioManager)
 
@@ -17,6 +19,7 @@ UTakeRecorderMicrophoneAudioManager::UTakeRecorderMicrophoneAudioManager(const F
 {
 }
 
+#if WITH_EDITOR
 TUniquePtr<IAudioCaptureEditor> UTakeRecorderMicrophoneAudioManager::CreateAudioRecorderObject()
 {
 	IAudioCaptureEditorModule& Factory = FModuleManager::Get().LoadModuleChecked<IAudioCaptureEditorModule>("AudioCaptureEditor");
@@ -37,13 +40,16 @@ void UTakeRecorderMicrophoneAudioManager::PostEditChangeProperty(FPropertyChange
 	// Notify the UI of the change so it can update if needed
 	GetOnAudioInputDeviceChanged().Broadcast();
 }
+#endif // WITH_EDITOR
 
 void UTakeRecorderMicrophoneAudioManager::EnumerateAudioDevices(bool InForceRefresh)
 {
+#if WITH_EDITOR
 	if (!AudioRecorder.IsValid() || InForceRefresh)
 	{
 		AudioRecorder = CreateAudioRecorderObject();
 	}
+#endif // WITH_EDITOR
 
 	if (AudioInputDevice.DeviceInfoArray.Num() == 0 || InForceRefresh)
 	{
@@ -65,6 +71,7 @@ void UTakeRecorderMicrophoneAudioManager::StartRecording(int32 InChannelCount)
 		return;
 	}
 	
+#if WITH_EDITOR
 	if (AudioRecorder.IsValid())
 	{
 		if (AudioRecorder->IsReadyToRecord())
@@ -83,16 +90,22 @@ void UTakeRecorderMicrophoneAudioManager::StartRecording(int32 InChannelCount)
 	{
 		UE_LOGF(LogTakesCore, Error, "Microphone Audio Source could not start. Please check that the AudioCapture plugin is enabled");
 	}
+#else // WITH_EDITOR
+	UE_LOG(LogTakesCore, Warning, TEXT("Microphone Audio Source is not currently supported in Cooked/Runtime builds. No microphone audio will be recorded."))
+#endif // WITH_EDITOR
 }
 
 void UTakeRecorderMicrophoneAudioManager::StopRecording()
 {
+#if WITH_EDITOR
 	if (AudioRecorder.IsValid() && AudioRecorder->IsRecording())
 	{
 		AudioRecorder->Stop();
 	}
+#endif // WITH_EDITOR
 }
 
+#if WITH_EDITOR
 TObjectPtr<USoundWave> UTakeRecorderMicrophoneAudioManager::GetRecordedSoundWave(const FTakeRecorderAudioSourceSettings& InSourceSettings)
 {
 	if (AudioRecorder.IsValid() && AudioRecorder->IsStopped())
@@ -102,9 +115,11 @@ TObjectPtr<USoundWave> UTakeRecorderMicrophoneAudioManager::GetRecordedSoundWave
 
 	return nullptr;
 }
+#endif // WITH_EDITOR
 
 void UTakeRecorderMicrophoneAudioManager::FinalizeRecording()
 {
+#if WITH_EDITOR
 	if (AudioRecorder.IsValid() && !AudioRecorder->IsReadyToRecord())
 	{
 		AudioRecorder.Reset();
@@ -115,6 +130,7 @@ void UTakeRecorderMicrophoneAudioManager::FinalizeRecording()
 		AudioRecorder = CreateAudioRecorderObject();
 		check(AudioRecorder.IsValid());
 	}
+#endif // WITH_EDITOR
 }
 
 int32 UTakeRecorderMicrophoneAudioManager::GetDeviceChannelCount()
@@ -153,16 +169,27 @@ void UTakeRecorderMicrophoneAudioManager::BuildDeviceInfoArray()
 {
 	AudioInputDevice.DeviceInfoArray.Empty();
 
+	FString DefaultDeviceId;
+	FString DefaultDeviceName;
+
+#if WITH_EDITOR
 	FTakeRecorderAudioDeviceInfo DefaultDeviceInfo;
 	const bool bFoundDefaultDevice = AudioRecorder->GetCaptureDeviceInfo(DefaultDeviceInfo);
+	DefaultDeviceId = DefaultDeviceInfo.DeviceId;
+	DefaultDeviceName = DefaultDeviceInfo.DeviceName;
+#else // WITH_EDITOR
+	const bool bFoundDefaultDevice = false;
+#endif // WITH_EDITOR
 
+#if WITH_EDITOR
 	TArray<FTakeRecorderAudioDeviceInfo> CaptureDevicesAvailable;
 	AudioRecorder->GetCaptureDevicesAvailable(CaptureDevicesAvailable);
 	for (const FTakeRecorderAudioDeviceInfo& DeviceInfo : CaptureDevicesAvailable)
 	{
-		const bool bIsDefaultDevice = bFoundDefaultDevice && (DeviceInfo.DeviceId == DefaultDeviceInfo.DeviceId);
+		const bool bIsDefaultDevice = bFoundDefaultDevice && (DeviceInfo.DeviceId == DefaultDeviceId);
 		AudioInputDevice.DeviceInfoArray.Add(FAudioInputDeviceInfoProperty(DeviceInfo.DeviceName, DeviceInfo.DeviceId, DeviceInfo.InputChannels, DeviceInfo.PreferredSampleRate, bIsDefaultDevice));
 	}
+#endif // WITH_EDITOR
 
 	// Device names often have numbers in them so perform
 	// an alpha-numeric sort.
@@ -220,13 +247,13 @@ void UTakeRecorderMicrophoneAudioManager::BuildDeviceInfoArray()
 		if (AudioInputDevice.DeviceId.IsEmpty())
 		{
 			// Default to using the default input device if not already set
-			AudioInputDevice.DeviceId = DefaultDeviceInfo.DeviceId;
+			AudioInputDevice.DeviceId = DefaultDeviceId;
 		}
 		else if (!IsAudioDeviceAvailable(AudioInputDevice.DeviceId))
 		{
 			// Revert to using the default input device if previously saved device is not available
-			AudioInputDevice.DeviceId = DefaultDeviceInfo.DeviceId;
-			UE_LOGF(LogTakesCore, Error, "Previously saved audio input device unavailable. Falling back to default device \"%ls\"", *DefaultDeviceInfo.DeviceName);
+			AudioInputDevice.DeviceId = DefaultDeviceId;
+			UE_LOGF(LogTakesCore, Error, "Previously saved audio input device unavailable. Falling back to default device \"%ls\"", *DefaultDeviceName);
 		}
 	}
 }
diff --git a/Engine/Plugins/VirtualProduction/Takes/Source/TakeRecorderSources/Private/TakeRecorderMicrophoneAudioSource.cpp b/Engine/Plugins/VirtualProduction/Takes/Source/TakeRecorderSources/Private/TakeRecorderMicrophoneAudioSource.cpp
--- a/Engine/Plugins/VirtualProduction/Takes/Source/TakeRecorderSources/Private/TakeRecorderMicrophoneAudioSource.cpp
+++ b/Engine/Plugins/VirtualProduction/Takes/Source/TakeRecorderSources/Private/TakeRecorderMicrophoneAudioSource.cpp
@@ -9,25 +9,25 @@
 #include "TakeRecorderSources.h"
 #include "TakesUtils.h"
 
-#include "AudioCaptureEditorTypes.h"
-#include "Editor.h"
 #include "LevelSequence.h"
 #include "Modules/ModuleManager.h"
 #include "MovieSceneFolder.h"
 #include "Recorder/TakeRecorderParameters.h"
 #include "Sections/MovieSceneAudioSection.h"
 #include "Sound/SoundWave.h"
-#include "Styling/SlateIconFinder.h"
 #include "Tracks/MovieSceneAudioTrack.h"
 
 #include "AssetRegistry/AssetData.h"
 #include "AssetRegistry/AssetRegistryModule.h"
 #include "Misc/PackageName.h"
-#include "ObjectEditorUtils.h"
-#include "ObjectTools.h"
-#include "PackageTools.h"
 #include "UObject/UObjectBaseUtility.h"
 
+#if WITH_EDITOR
+#include "Editor.h"
+#include "AudioCaptureEditorTypes.h"
+#include "ObjectTools.h"
+#endif // WITH_EDITOR
+
 #include UE_INLINE_GENERATED_CPP_BY_NAME(TakeRecorderMicrophoneAudioSource)
 
 UTakeRecorderMicrophoneAudioSourceSettings::UTakeRecorderMicrophoneAudioSourceSettings(const FObjectInitializer& ObjInit)
@@ -41,6 +41,7 @@ UTakeRecorderMicrophoneAudioSourceSettings::UTakeRecorderMicrophoneAudioSourceSe
 	AudioSourceName = FText::Format(FText::FromString(TEXT("{0} {1}")), AudioSourceName, FText::FromString(TEXT("{channel}")));
 }
 
+#if WITH_EDITOR
 void UTakeRecorderMicrophoneAudioSourceSettings::PostEditChangeProperty(FPropertyChangedEvent& PropertyChangedEvent)
 {
 	Super::PostEditChangeProperty(PropertyChangedEvent);
@@ -50,9 +51,11 @@ void UTakeRecorderMicrophoneAudioSourceSettings::PostEditChangeProperty(FPropert
 		SaveConfig();
 	}
 }
+#endif // WITH_EDITOR
 
 FString UTakeRecorderMicrophoneAudioSourceSettings::GetSubsceneTrackName(ULevelSequence* InSequence) const
 {
+#if WITH_EDITOR
 	if (UTakeMetaData* TakeMetaData = InSequence->FindMetaData<UTakeMetaData>())
 	{
 		// Before Naming Tokens we were always evaluating AudioTrackName_{slate}.
@@ -60,15 +63,18 @@ FString UTakeRecorderMicrophoneAudioSourceSettings::GetSubsceneTrackName(ULevelS
 		const FString SubSceneTrackName = FString::Printf(TEXT("%s_%s"), *AudioTrackName.ToString(), *TakeMetaData->GenerateAssetPath("{slate}"));
 		return *TakeMetaData->GenerateAssetPath(SubSceneTrackName);
 	}
+#endif // WITH_EDITOR
 	return TEXT("MicrophoneAudio");
 }
 
 FString UTakeRecorderMicrophoneAudioSourceSettings::GetSubsceneAssetName(ULevelSequence* InSequence) const
 {
+#if WITH_EDITOR
 	if (UTakeMetaData* TakeMetaData = InSequence->FindMetaData<UTakeMetaData>())
 	{
 		return TakeMetaData->GenerateAssetPath(AudioAssetName);
 	}
+#endif // WITH_EDITOR
 	return TEXT("MicrophoneAudio");
 }
 
@@ -114,11 +120,13 @@ void UTakeRecorderMicrophoneAudioSource::Initialize()
 	DoAudioSetup();
 }
 
+#if WITH_EDITOR
 void UTakeRecorderMicrophoneAudioSource::PostEditChangeProperty(FPropertyChangedEvent& PropertyChangedEvent)
 {
 	DoAudioSetup();
 	Super::PostEditChangeProperty(PropertyChangedEvent);
 }
+#endif // WITH_EDITOR
 
 UTakeRecorderMicrophoneAudioManager* UTakeRecorderMicrophoneAudioSource::GetAudioInputManager()
 {
@@ -184,7 +192,7 @@ FString UTakeRecorderMicrophoneAudioSource::GetAudioAssetName(ULevelSequence* In
 	FString EvaluatedAudioAssetName;
 	if (TryEvaluateTokenString(AudioAssetName, InSequence, EvaluatedAudioAssetName))
 	{
-		return UPackageTools::SanitizePackageName(EvaluatedAudioAssetName);
+		return TakesUtils::SanitizePackageName(EvaluatedAudioAssetName);
 	}
 
 	return TEXT("MicrophoneAudio");
@@ -245,6 +253,9 @@ TArray<UTakeRecorderSource*> UTakeRecorderMicrophoneAudioSource::PreRecording(UL
 		FString TrackName = GetAudioTrackName(InSequence);
 		UMovieScene* MovieScene = InSequence->GetMovieScene();
 		
+		// To review this. Ultimately currently we do not support microphone audio recording in runtime, but theres no reason this shouldnt work better
+		// without using the display name.
+	#if WITH_EDITOR
 		for (auto Track : ReverseIterate(MovieScene->GetTracks()))
 		{
 			if (Track->IsA(UMovieSceneAudioTrack::StaticClass()) && Track->GetDisplayName().EqualTo(FText::FromString(TrackName)))
@@ -257,11 +268,14 @@ TArray<UTakeRecorderSource*> UTakeRecorderMicrophoneAudioSource::PreRecording(UL
 				}
 			}
 		}
+	#endif // WITH_EDITOR
 
 		if (!CachedAudioTrack.IsValid())
 		{
 			CachedAudioTrack = MovieScene->AddTrack<UMovieSceneAudioTrack>();
+		#if WITH_EDITOR
 			CachedAudioTrack->SetDisplayName(FText::FromString(TrackName));
+		#endif // WITH_EDITOR
 		}
 
 		if (bReplaceRecordedAudio)
@@ -291,7 +305,7 @@ TArray<UTakeRecorderSource*> UTakeRecorderMicrophoneAudioSource::PreRecording(UL
 			FString TempStr = AudioSubDirectory;
 			if (TryEvaluateTokenString(TempStr, InSequence, TempStr))
 			{
-				TempStr = UPackageTools::SanitizePackageName(TempStr);
+				TempStr = TakesUtils::SanitizePackageName(TempStr);
 			}
 			AudioDirectory.Path /= TempStr;
 		}
@@ -422,7 +436,13 @@ TArray<UTakeRecorderSource*> UTakeRecorderMicrophoneAudioSource::PostRecording(U
 					AudioSection->TimecodeSource = StartTimecode;
 				}
 
-				if ((Sources->GetSettings().bSaveRecordedAssets) || GEditor == nullptr)
+			#if WITH_EDITOR
+				const bool bEditorNotAvailable = GEditor == nullptr;
+			#else // WITH_EDITOR
+				const bool bEditorNotAvailable = true;
+			#endif // WITH_EDITOR
+
+				if ((Sources->GetSettings().bSaveRecordedAssets) || bEditorNotAvailable)
 				{
 					const bool bIsPossessable = TakesUtils::IsSourcePossessable(this);
 					TakesUtils::SaveAsset(SoundWave, bIsPossessable);
@@ -437,10 +457,12 @@ TArray<UTakeRecorderSource*> UTakeRecorderMicrophoneAudioSource::PostRecording(U
 	// Reset our sound wave pointer
 	RecordedSoundWave.Reset();
 	
+#if WITH_EDITOR
 	if (GEditor && AssetsToCleanUp.Num() > 0)
 	{
 		ObjectTools::ForceDeleteObjects(AssetsToCleanUp, false);
 	}
+#endif // WITH_EDITOR
 
 	return TArray<UTakeRecorderSource*>();
 }
@@ -483,6 +505,7 @@ void UTakeRecorderMicrophoneAudioSource::GetRecordedSoundWave(ULevelSequence* In
 		UMovieScene* MovieScene = InSequence->GetMovieScene();
 		FFrameRate DisplayRate = MovieScene->GetDisplayRate();
 
+	#if WITH_EDITOR
 		FTakeRecorderAudioSourceSettings AudioSettings;
 		AudioSettings.Directory = AudioDirectory;
 		AudioSettings.AssetName = AssetFileName;
@@ -502,6 +525,7 @@ void UTakeRecorderMicrophoneAudioSource::GetRecordedSoundWave(ULevelSequence* In
 			
 			RecordedSoundWave = SoundWave;
 		}
+	#endif // WITH_EDITOR
 	}
 }
 
@@ -588,6 +612,7 @@ FString UTakeRecorderMicrophoneAudioSource::CreateUniqueAudioAssetName(ULevelSeq
 
 bool UTakeRecorderMicrophoneAudioSource::TryEvaluateTokenString(const FString& InString, const ULevelSequence* InSequence, FString& OutString) const
 {
+#if WITH_EDITOR
 	if (const UTakeMetaData* TakeMetaData = InSequence->FindMetaData<UTakeMetaData>())
 	{
 		UTakeRecorderNamingTokensContext* Context = NewObject<UTakeRecorderNamingTokensContext>();
@@ -595,5 +620,6 @@ bool UTakeRecorderMicrophoneAudioSource::TryEvaluateTokenString(const FString& I
 		OutString = TakeMetaData->GenerateAssetPath(InString, Context);
 		return true;
 	}
+#endif // WITH_EDITOR
 	return false;
 }
diff --git a/Engine/Plugins/VirtualProduction/Takes/Source/TakeRecorderSources/Private/TakeRecorderNearbySpawnedActorSource.cpp b/Engine/Plugins/VirtualProduction/Takes/Source/TakeRecorderSources/Private/TakeRecorderNearbySpawnedActorSource.cpp
--- a/Engine/Plugins/VirtualProduction/Takes/Source/TakeRecorderSources/Private/TakeRecorderNearbySpawnedActorSource.cpp
+++ b/Engine/Plugins/VirtualProduction/Takes/Source/TakeRecorderSources/Private/TakeRecorderNearbySpawnedActorSource.cpp
@@ -2,20 +2,16 @@
 
 #include "TakeRecorderNearbySpawnedActorSource.h"
 
-#include "AssetRegistry/AssetData.h"
 #include "GameFramework/Pawn.h"
-#include "TakeRecorderSources.h"
 #include "TakeRecorderSource.h"
 #include "TakeRecorderActorSource.h"
 #include "TakeRecorderSourcesUtils.h"
-#include "TakesUtils.h"
 #include "TakesCoreLog.h"
 
 #include "LevelSequence.h"
 #include "Engine/Engine.h"
 #include "GameFramework/PlayerController.h"
 
-#include "Styling/SlateIconFinder.h"
 #include "Math/UnitConversion.h"
 
 #include UE_INLINE_GENERATED_CPP_BY_NAME(TakeRecorderNearbySpawnedActorSource)
diff --git a/Engine/Plugins/VirtualProduction/Takes/Source/TakeRecorderSources/Private/TakeRecorderPlayerSource.cpp b/Engine/Plugins/VirtualProduction/Takes/Source/TakeRecorderSources/Private/TakeRecorderPlayerSource.cpp
--- a/Engine/Plugins/VirtualProduction/Takes/Source/TakeRecorderSources/Private/TakeRecorderPlayerSource.cpp
+++ b/Engine/Plugins/VirtualProduction/Takes/Source/TakeRecorderSources/Private/TakeRecorderPlayerSource.cpp
@@ -1,7 +1,6 @@
 // Copyright Epic Games, Inc. All Rights Reserved.
 
 #include "TakeRecorderPlayerSource.h"
-#include "AssetRegistry/AssetData.h"
 #include "GameFramework/Pawn.h"
 #include "TakesUtils.h"
 
@@ -11,7 +10,6 @@
 #include "LevelSequence.h"
 #include "Engine/Engine.h"
 #include "GameFramework/PlayerController.h"
-#include "Modules/ModuleManager.h"
 
 #include UE_INLINE_GENERATED_CPP_BY_NAME(TakeRecorderPlayerSource)
 
@@ -60,12 +58,8 @@ TArray<UTakeRecorderSource*> UTakeRecorderPlayerSource::PreRecording(ULevelSeque
 	}
 
 	UTakeRecorderActorSource* ActorSource = NewObject<UTakeRecorderActorSource>(Sources, UTakeRecorderActorSource::StaticClass(), NAME_None, RF_Transactional);
-	ActorSource->Target = CurrentPlayer;
+	ActorSource->SetSourceActor(CurrentPlayer);
 
-	// Send a PropertyChangedEvent so the class catches the callback and rebuilds the property map. We can't rely on the Actor rebuilding the map on PreRecording
-	// because that would wipe out any user adjustments from one added natively.
-	FPropertyChangedEvent PropertyChangedEvent(UTakeRecorderActorSource::StaticClass()->FindPropertyByName(GET_MEMBER_NAME_CHECKED(UTakeRecorderActorSource, Target)), EPropertyChangeType::ValueSet);
-	ActorSource->PostEditChangeProperty(PropertyChangedEvent);
 	NewSources.Add(ActorSource);
 
 	PlayerActorSource = ActorSource;
diff --git a/Engine/Plugins/VirtualProduction/Takes/Source/TakeRecorderSources/Private/TakeRecorderSourceHelpers.cpp b/Engine/Plugins/VirtualProduction/Takes/Source/TakeRecorderSources/Private/TakeRecorderSourceHelpers.cpp
--- a/Engine/Plugins/VirtualProduction/Takes/Source/TakeRecorderSources/Private/TakeRecorderSourceHelpers.cpp
+++ b/Engine/Plugins/VirtualProduction/Takes/Source/TakeRecorderSources/Private/TakeRecorderSourceHelpers.cpp
@@ -5,14 +5,14 @@
 #include "TakeRecorderActorSource.h"
 #include "LevelSequenceActor.h"
 #include "GameFramework/Actor.h"
-#include "ScopedTransaction.h"
 #include "TakeRecorderLevelSequenceSource.h"
 #include "TakeRecorderTimeProcessing.h"
-
 #include "LevelSequence.h"
-#include "MovieSceneTakeSection.h"
 
+#if WITH_EDITOR
 #include "Editor.h"
+#include "ScopedTransaction.h"
+#endif // WITH_EDITOR
 
 #define LOCTEXT_NAMESPACE "TakeRecorderSourceHelpers"
 
@@ -27,8 +27,10 @@ void AddActorSources(
 {
 	if (ActorsToRecord.Num() > 0)
 	{
+	#if WITH_EDITOR
 		FScopedTransaction Transaction(FText::Format(
 			LOCTEXT("AddSources", "Add Recording {0}|plural(one=Source, other=Sources)"), ActorsToRecord.Num()));
+	#endif // WITH_EDITOR
 		TakeRecorderSources->Modify();
 
 		for (AActor* Actor : ActorsToRecord)
@@ -69,21 +71,19 @@ void AddActorSources(
 			else
 			{
 				UTakeRecorderActorSource* NewSource = TakeRecorderSources->AddSource<UTakeRecorderActorSource>();
+				AActor* TargetActor = Actor;
 
+			#if WITH_EDITOR
 				if (AActor* EditorActor = EditorUtilities::GetEditorWorldCounterpartActor(Actor))
 				{
-					NewSource->Target = EditorActor;
-				}
-				else
-				{
-					NewSource->Target = Actor;
+					TargetActor = EditorActor;
 				}
+			#endif // WITH_EDITOR
+				
+				NewSource->SetSourceActor(TargetActor);
+
 				NewSource->bShowProgressDialog = bShowProgress;
 				NewSource->bReduceKeys = bReduceKeys;
-
-				// Send a PropertyChangedEvent so the class catches the callback and rebuilds the property map.
-				FPropertyChangedEvent PropertyChangedEvent(UTakeRecorderActorSource::StaticClass()->FindPropertyByName(GET_MEMBER_NAME_CHECKED(UTakeRecorderActorSource, Target)), EPropertyChangeType::ValueSet);
-				NewSource->PostEditChangeProperty(PropertyChangedEvent);
 			}
 		}
 	}
@@ -96,9 +96,11 @@ void RemoveActorSources(UTakeRecorderSources* TakeRecorderSources, TArrayView<AA
 		return;
 	}
 
+#if WITH_EDITOR
 	FScopedTransaction Transaction(FText::Format(LOCTEXT(
 			"RemoveActorSources", 
 			"Remove Recording {0}|plural(one=Source, other=Sources)"), ActorsToRemove.Num()));
+#endif // WITH_EDITOR
 	TakeRecorderSources->Modify();
 
 	for (UTakeRecorderSource* Source : TakeRecorderSources->GetSourcesCopy())
@@ -120,9 +122,11 @@ void RemoveAllActorSources(UTakeRecorderSources* TakeRecorderSources)
 		return;
 	}
 
+#if WITH_EDITOR
 	FScopedTransaction Transaction(FText::Format(LOCTEXT(
 			"RemoveAllActorSources", 
 			"Remove Recording {0}|plural(one=Source, other=Sources)"), TakeRecorderSources->GetSources().Num()));
+#endif // WITH_EDITOR
 	TakeRecorderSources->Modify();
 
 	while (!TakeRecorderSources->GetSources().IsEmpty())
diff --git a/Engine/Plugins/VirtualProduction/Takes/Source/TakeRecorderSources/Private/TakeRecorderSourcesModule.cpp b/Engine/Plugins/VirtualProduction/Takes/Source/TakeRecorderSources/Private/TakeRecorderSourcesModule.cpp
--- a/Engine/Plugins/VirtualProduction/Takes/Source/TakeRecorderSources/Private/TakeRecorderSourcesModule.cpp
+++ b/Engine/Plugins/VirtualProduction/Takes/Source/TakeRecorderSources/Private/TakeRecorderSourcesModule.cpp
@@ -6,14 +6,11 @@
 #include "Modules/ModuleManager.h"
 #include "ITakeRecorderModule.h"
 #include "ITakeRecorderDropHandler.h"
-#include "TakeRecorderLevelSequenceSource.h"
-#include "Input/DragAndDrop.h"
 #include "Recorder/TakeRecorder.h"
 #include "Recorder/TakeRecorderBlueprintLibrary.h"
 #include "Recorder/TakeRecorderPanel.h"
 #include "Recorder/TakeRecorderParameters.h"
 #include "TakesCoreBlueprintLibrary.h"
-#include "TakesCoreLog.h"
 #include "TakeMetaData.h"
 #include "TakeRecorderSourceHelpers.h"
 #include "TakeRecorderActorSource.h"
@@ -23,44 +20,46 @@
 #include "TakeRecorderSettings.h"
 #include "TakeRecorderSourcesUtils.h"
 #include "Features/IModularFeatures.h"
-#include "DragAndDrop/ActorDragDropOp.h"
-#include "DragAndDrop/FolderDragDropOp.h"
-#include "DragAndDrop/CompositeDragDropOp.h"
-#include "EngineUtils.h"
-#include "Algo/Sort.h"
-#include "ScopedTransaction.h"
+
 #include "LevelSequenceActor.h"
-#include "LevelSequenceEditorBlueprintLibrary.h"
+#include "LevelSequence.h"
 
 #include "Framework/Application/SlateApplication.h"
 #include "Framework/MultiBox/MultiBoxBuilder.h"
 #include "Framework/MultiBox/MultiBoxExtender.h"
 #include "Widgets/Layout/SBox.h"
-#include "Dialogs/Dialogs.h"
+#include "Styling/SlateIconFinder.h"
 
 #include "Engine/LevelScriptActor.h"
-#include "Engine/Selection.h"
-#include "Editor.h"
-#include "LevelEditor.h"
-#include "SceneOutlinerModule.h"
-#include "SceneOutlinerPublicTypes.h"
 #include "Framework/Notifications/NotificationManager.h"
 #include "Widgets/Notifications/SNotificationList.h"
-#include "Subsystems/AssetEditorSubsystem.h"
-#include "ActorTreeItem.h"
-#include "EditorActorFolders.h"
 
 #include "TakeRecorderMicrophoneAudioSource.h"
 #include "TakeRecorderWorldSource.h"
 #include "TrackRecorders/MovieSceneAnimationTrackRecorderSettings.h"
 
-#include "ISequencer.h"
-#include "ISequencerModule.h"
-#include "ILevelSequenceEditorToolkit.h"
 #include "ITakeRecorderSourcesModule.h"
 
+#include "Bindings/MovieSceneCustomBinding.h"
 #include "Bindings/MovieSceneSpawnableActorBinding.h"
 
+#if WITH_EDITOR
+#include "LevelSequenceEditorBlueprintLibrary.h"
+
+#include "Dialogs/Dialogs.h"
+#include "Engine/Selection.h"
+
+#include "Editor.h"
+#include "LevelEditor.h"
+#include "SceneOutlinerModule.h"
+#include "SceneOutlinerPublicTypes.h"
+#include "Subsystems/AssetEditorSubsystem.h"
+#include "ActorTreeItem.h"
+
+#include "ISequencer.h"
+#include "ISequencerModule.h"
+#endif // WITH_EDITOR
+
 #define LOCTEXT_NAMESPACE "TakeRecorderSources"
 
 namespace UE::TakeRecorderSources
@@ -73,6 +72,8 @@ namespace UE::TakeRecorderSources
 
 namespace
 {
+#if WITH_EDITOR
+	// ActorLabel is an editor time only concept.
 	static AActor* FindActorByLabel(const FString& ActorNameStr, UWorld* InWorld, bool bFuzzy = false)
 	{
 		// search for the actor by name
@@ -117,6 +118,7 @@ namespace
 
 		return nullptr;
 	}
+#endif // WITH_EDITOR
 
 	static void FindActorsOfClass(UClass* Class, UWorld* InWorld, TArray<AActor*>& OutActors)
 	{
@@ -150,7 +152,9 @@ void FTakeRecorderSourcesModule::StartupModule()
 
 	ITakeRecorderModule& TakeRecorderModule = FModuleManager::Get().LoadModuleChecked<ITakeRecorderModule>("TakeRecorder");
 
+#if WITH_EDITOR
 	SourcesMenuExtension = TakeRecorderModule.RegisterSourcesMenuExtension(FOnExtendSourcesMenu::CreateStatic(ExtendSourcesMenu));
+#endif // WITH_EDITOR
 
 	ITakeRecorderModule::FSourceExtensionData ExtensionData;
 	ExtensionData.OnAddActorSource = ITakeRecorderModule::FSourceExtensionData::FAddActorSource::CreateStatic(TakeRecorderSourceHelpers::AddActorSources);
@@ -163,8 +167,10 @@ void FTakeRecorderSourcesModule::StartupModule()
 	TakeRecorderModule.RegisterSettingsObject(GetMutableDefault<UMovieSceneAnimationTrackRecorderEditorSettings>());
 	TakeRecorderModule.RegisterSettingsObject(GetMutableDefault<UTakeRecorderWorldSourceSettings>());
 
+#if WITH_EDITOR
 	ISequencerModule& SequencerModule = FModuleManager::Get().LoadModuleChecked<ISequencerModule>("Sequencer");
 	OnSequencerCreatedHandle = SequencerModule.RegisterOnSequencerCreated(FOnSequencerCreated::FDelegate::CreateRaw(this, &FTakeRecorderSourcesModule::OnSequencerCreated));
+#endif // WITH_EDITOR
 	
 	UTakeRecorderSources::OnAddSourceForActorBinding().BindRaw(this, &FTakeRecorderSourcesModule::OnNewActorReferenceBindingAdded);
 }
@@ -173,11 +179,13 @@ void FTakeRecorderSourcesModule::ShutdownModule()
 {
 	IModularFeatures::Get().UnregisterModularFeature(ITakeRecorderDropHandler::ModularFeatureName, &ActorDropHandler);
 
+#if WITH_EDITOR
 	ISequencerModule* SequencerModule = FModuleManager::GetModulePtr<ISequencerModule>("Sequencer");
 	if (SequencerModule)
 	{
 		SequencerModule->UnregisterOnSequencerCreated(OnSequencerCreatedHandle);
 	}
+#endif // WITH_EDITOR
 	
 	ITakeRecorderModule* TakeRecorderModule = FModuleManager::Get().GetModulePtr<ITakeRecorderModule>("TakeRecorder");
 	if (TakeRecorderModule)
@@ -227,6 +235,7 @@ void FTakeRecorderSourcesModule::RegisterMenuExtensions()
 
 void FTakeRecorderSourcesModule::UnregisterMenuExtensions()
 {
+#if WITH_EDITOR
 	// Unregister level editor menu extender
 	if (FModuleManager::Get().IsModuleLoaded(TEXT("LevelEditor")))
 	{
@@ -235,6 +244,7 @@ void FTakeRecorderSourcesModule::UnregisterMenuExtensions()
 			return Delegate.GetHandle() == LevelEditorExtenderDelegateHandle;
 		});
 	}
+#endif // WITH_EDITOR
 }
 
 void FTakeRecorderSourcesModule::BindCommands()
@@ -264,7 +274,11 @@ TSharedRef<FExtender> FTakeRecorderSourcesModule::ExtendLevelViewportContextMenu
 			FText RecordText;
 			if (SelectedActors.Num() == 1)
 			{
+			#if WITH_EDITOR
 				RecordText = FText::Format(LOCTEXT("RecordSingleSelectedActorText", "Record {0} with Take Recorder"), FText::FromString(SelectedActors[0]->GetActorLabel()));
+			#else // WITH_EDITOR
+				RecordText = FText::Format(LOCTEXT("RecordSingleSelectedActorText", "Record {0} with Take Recorder"), FText::FromString(SelectedActors[0]->GetName()));
+			#endif // WITH_EDITOR
 			}
 			else
 			{
@@ -280,6 +294,7 @@ TSharedRef<FExtender> FTakeRecorderSourcesModule::ExtendLevelViewportContextMenu
 	return Extender;
 }
 
+#if WITH_EDITOR
 void FTakeRecorderSourcesModule::ExtendSourcesMenu(TSharedRef<FExtender> Extender, UTakeRecorderSources* Sources)
 {
 	Extender->AddMenuExtension("Sources", EExtensionHook::Before, nullptr, FMenuExtensionDelegate::CreateStatic(PopulateSourcesMenu, Sources));
@@ -386,6 +401,7 @@ void FTakeRecorderSourcesModule::PopulateActorSubMenu(FMenuBuilder& MenuBuilder,
 	}
 	MenuBuilder.EndSection();
 }
+#endif // WITH_EDITOR
 
 bool FTakeRecorderSourcesModule::Exec_Editor(UWorld* InWorld, const TCHAR* Cmd, FOutputDevice& Ar)
 {
@@ -656,6 +672,7 @@ void FTakeRecorderSourcesModule::RecordActors(const TArray<AActor*>& ActorsToRec
 
 void FTakeRecorderSourcesModule::RecordSelectedActors()
 {
+#if WITH_EDITOR
 	TArray<AActor*> SelectedActors;
 	GEditor->GetSelectedActors()->GetSelectedObjects(SelectedActors);
 
@@ -671,8 +688,10 @@ void FTakeRecorderSourcesModule::RecordSelectedActors()
 	ULevelSequence* LevelSequence = ULevelSequenceEditorBlueprintLibrary::GetFocusedLevelSequence();
 
 	RecordActors(SelectedActors, LevelSequence, RootLevelSequence);
+#endif // WITH_EDITOR
 }
 
+#if WITH_EDITOR
 void FTakeRecorderSourcesModule::OnSequencerCreated(TSharedRef<ISequencer> Sequencer)
 {
 	bool bSupportsSpawnableActorBinding = false;
@@ -776,6 +795,7 @@ void FTakeRecorderSourcesModule::OnSequencerCreated(TSharedRef<ISequencer> Seque
 		}
 	});
 }
+#endif // WITH_EDITOR
 
 UTakeRecorderSource* FTakeRecorderSourcesModule::OnNewActorReferenceBindingAdded(AActor* InActor, ULevelSequence* InLevelSequence)
 {	
diff --git a/Engine/Plugins/VirtualProduction/Takes/Source/TakeRecorderSources/Private/TakeRecorderSourcesModule.h b/Engine/Plugins/VirtualProduction/Takes/Source/TakeRecorderSources/Private/TakeRecorderSourcesModule.h
--- a/Engine/Plugins/VirtualProduction/Takes/Source/TakeRecorderSources/Private/TakeRecorderSourcesModule.h
+++ b/Engine/Plugins/VirtualProduction/Takes/Source/TakeRecorderSources/Private/TakeRecorderSourcesModule.h
@@ -1,20 +1,27 @@
-﻿// Copyright Epic Games, Inc. All Rights Reserved.
+// Copyright Epic Games, Inc. All Rights Reserved.
 
 #pragma once
 
 #include "ActorTakeRecorderDropHandler.h"
 #include "ITakeRecorderSourcesModule.h"
-#include "LevelEditor.h"
 #include "Misc/CoreMisc.h"
 #include "Templates/SharedPointer.h"
 
-class ISequencer;
-template<typename OptionalType>struct TOptional;
+#if WITH_EDITOR
+#include "LevelEditor.h"
+#endif // WITH_EDITOR
+
+template<typename OptionalType> struct TOptional;
 
+class FExtender;
 class ULevelSequence;
 class UTakeRecorderSources;
 class UTakeRecorderSource;
 
+#if WITH_EDITOR
+class ISequencer;
+#endif // WITH_EDITOR
+
 namespace UE::TakeRecorderSources
 {
 class FTakeRecorderSourcesModule : public ITakeRecorderSourcesModule, private FSelfRegisteringExec
@@ -37,11 +44,17 @@ public:
 
 private:
 	
+#if WITH_EDITOR
 	FLevelEditorModule::FLevelViewportMenuExtender_SelectedActors LevelEditorMenuExtenderDelegate;
+#endif // WITH_EDITOR
 
 	FActorTakeRecorderDropHandler ActorDropHandler;
 	FDelegateHandle SourcesMenuExtension;
+
+#if WITH_EDITOR
 	FDelegateHandle LevelEditorExtenderDelegateHandle;
+#endif // WITH_EDITOR
+
 	FDelegateHandle OnSequencerCreatedHandle;
 
 	TSharedPtr<FUICommandList> CommandList;
@@ -54,9 +67,11 @@ private:
 
 	TSharedRef<FExtender> ExtendLevelViewportContextMenu(const TSharedRef<FUICommandList> InCommandList, const TArray<AActor*> SelectedActors);
 
+#if WITH_EDITOR
 	static void ExtendSourcesMenu(TSharedRef<FExtender> Extender, UTakeRecorderSources* Sources);
 	static void PopulateSourcesMenu(FMenuBuilder& MenuBuilder, UTakeRecorderSources* Sources);
 	static void PopulateActorSubMenu(FMenuBuilder& MenuBuilder, UTakeRecorderSources* Sources);
+#endif // WITH_EDITOR
 
 	//~ Begin FSelfRegisteringExec Interface
 	virtual bool Exec_Editor(UWorld* InWorld, const TCHAR* Cmd, FOutputDevice& Ar) override;
@@ -68,7 +83,9 @@ private:
 	void RecordActors(const TArray<AActor*>& ActorsToRecord, TOptional<ULevelSequence*> LevelSequence, TOptional<ULevelSequence*> RootLevelSequence);
 	void RecordSelectedActors();
 
+#if WITH_EDITOR
 	void OnSequencerCreated(TSharedRef<ISequencer> Sequencer);
+#endif // WITH_EDITOR
 	
 	UTakeRecorderSource* OnNewActorReferenceBindingAdded(AActor* InActor, ULevelSequence* InLevelSequence);
 };
diff --git a/Engine/Plugins/VirtualProduction/Takes/Source/TakeRecorderSources/Private/TakeRecorderSourcesUtils.h b/Engine/Plugins/VirtualProduction/Takes/Source/TakeRecorderSources/Private/TakeRecorderSourcesUtils.h
--- a/Engine/Plugins/VirtualProduction/Takes/Source/TakeRecorderSources/Private/TakeRecorderSourcesUtils.h
+++ b/Engine/Plugins/VirtualProduction/Takes/Source/TakeRecorderSources/Private/TakeRecorderSourcesUtils.h
@@ -2,10 +2,6 @@
 
 #pragma once
 
-#include "CoreMinimal.h"
-#include "Modules/ModuleManager.h"
-#include "UObject/UObjectGlobals.h"
-
 #include "TakeRecorderActorSource.h"
 #include "TakeRecorderSources.h"
 #include "TakesUtils.h"
@@ -13,10 +9,8 @@
 #include "TrackRecorders/MovieSceneTrackRecorder.h"
 #include "TrackRecorders/MovieSceneAnimationTrackRecorder.h"
 
-#include "LevelSequence.h"
-#include "MovieScene.h"
 #include "MovieSceneSequenceID.h"
-#include "Compilation/MovieSceneCompiledDataManager.h"
+#include "LevelSequence.h"
 
 #include "ITakeRecorderSourcesManager.h"
 
@@ -226,12 +220,8 @@ namespace TakeRecorderSourcesUtils
 		UTakeRecorderSources* Sources = ITakeRecorderSourcesManager::GetChecked().FindOrAddSources(InLevelSequence);
 
 		UTakeRecorderActorSource* ActorSource = Sources->AddSource<UTakeRecorderActorSource>();
-		ActorSource->Target = InActor;
-
-		// Send a PropertyChangedEvent so the class catches the callback and rebuilds the property map. We can't rely on the Actor rebuilding the map on PreRecording
-		// because that would wipe out any user adjustments from one added natively.
-		FPropertyChangedEvent PropertyChangedEvent(UTakeRecorderActorSource::StaticClass()->FindPropertyByName(GET_MEMBER_NAME_CHECKED(UTakeRecorderActorSource, Target)), EPropertyChangeType::ValueSet);
-		ActorSource->PostEditChangeProperty(PropertyChangedEvent);
+		// Setting the source actor through the function ensures internal property maps are also updated.
+		ActorSource->SetSourceActor(InActor);
 
 		// This has to be called after setting the Target and propagating the change event so that it has a chance to know what to record
 		// about the actor.
diff --git a/Engine/Plugins/VirtualProduction/Takes/Source/TakeRecorderSources/Private/TakeRecorderWorldSource.cpp b/Engine/Plugins/VirtualProduction/Takes/Source/TakeRecorderSources/Private/TakeRecorderWorldSource.cpp
--- a/Engine/Plugins/VirtualProduction/Takes/Source/TakeRecorderSources/Private/TakeRecorderWorldSource.cpp
+++ b/Engine/Plugins/VirtualProduction/Takes/Source/TakeRecorderSources/Private/TakeRecorderWorldSource.cpp
@@ -2,21 +2,20 @@
 
 #include "TakeRecorderWorldSource.h"
 
-#include "AssetRegistry/AssetData.h"
 #include "TakeRecorderSources.h"
 #include "TakeRecorderActorSource.h"
 #include "TakeRecorderSourcesUtils.h"
-#include "TakesUtils.h"
 
 #include "LevelSequence.h"
 #include "GameFramework/WorldSettings.h"
 
-#include "ISequencer.h"
-#include "MovieSceneFolder.h"
 #include "ITakeRecorderSourcesManager.h"
 
+#if WITH_EDITOR
+#include "ISequencer.h"
 #include "ILevelSequenceEditorToolkit.h"
 #include "Subsystems/AssetEditorSubsystem.h"
+#endif // WITH_EDITOR
 
 #include UE_INLINE_GENERATED_CPP_BY_NAME(TakeRecorderWorldSource)
 
@@ -29,6 +28,7 @@ UTakeRecorderWorldSourceSettings::UTakeRecorderWorldSourceSettings(const FObject
 {
 }
 
+#if WITH_EDITOR
 void UTakeRecorderWorldSourceSettings::PostEditChangeProperty(FPropertyChangedEvent& PropertyChangedEvent)
 {
 	Super::PostEditChangeProperty(PropertyChangedEvent);
@@ -38,6 +38,7 @@ void UTakeRecorderWorldSourceSettings::PostEditChangeProperty(FPropertyChangedEv
 		SaveConfig();
 	}
 }
+#endif // WITH_EDITOR
 
 UTakeRecorderWorldSource::UTakeRecorderWorldSource(const FObjectInitializer& ObjInit)
 	: Super(ObjInit)
@@ -84,12 +85,8 @@ TArray<UTakeRecorderSource*> UTakeRecorderWorldSource::PreRecording(ULevelSequen
 		}
 
 		UTakeRecorderActorSource* ActorSource = NewObject<UTakeRecorderActorSource>(Sources, UTakeRecorderActorSource::StaticClass(), NAME_None, RF_Transactional);
-		ActorSource->Target = WorldSettings;
-
-		// Send a PropertyChangedEvent so the class catches the callback and rebuilds the property map. We can't rely on the Actor rebuilding the map on PreRecording
-		// because that would wipe out any user adjustments from one added natively.
-		FPropertyChangedEvent PropertyChangedEvent(UTakeRecorderActorSource::StaticClass()->FindPropertyByName(GET_MEMBER_NAME_CHECKED(UTakeRecorderActorSource, Target)), EPropertyChangeType::ValueSet);
-		ActorSource->PostEditChangeProperty(PropertyChangedEvent);
+		// Setting the source actor through the function ensures internal property maps are also updated.
+		ActorSource->SetSourceActor(WorldSettings);
 
 		NewSources.Add(ActorSource);
 		WorldSource = ActorSource;
@@ -154,11 +151,20 @@ bool UTakeRecorderWorldSource::CanAddSource(UTakeRecorderSources* InSources) con
 
 void UTakeRecorderWorldSource::AutotrackActors(class ULevelSequence* InSequence, UWorld* InWorld)
 {
-	if (!InWorld || !GEditor)
+	if (!InWorld)
+	{
+		return;
+	}
+
+#if WITH_EDITOR
+	if (!GEditor)
 	{
 		return;
 	}
 
+	// To review - this whole chunk seems a little pointless as we dont use it at all other than to check if a ISequencer is valid.
+	// This is not though used elsewhere, and is bringing in Editor dependencies. I wonder if we need this at all as we look to move away
+	// from ISequencer?
 	IAssetEditorInstance*        AssetEditor = GEditor->GetEditorSubsystem<UAssetEditorSubsystem>()->FindEditorForAsset(InSequence, false);
 	ILevelSequenceEditorToolkit* LevelSequenceEditor = static_cast<ILevelSequenceEditorToolkit*>(AssetEditor);
 
@@ -168,6 +174,7 @@ void UTakeRecorderWorldSource::AutotrackActors(class ULevelSequence* InSequence,
 	{
 		return;
 	}
+#endif // WITH_EDITOR
 
 	UTakeRecorderSources* Sources = ITakeRecorderSourcesManager::GetChecked().FindOrAddSources(InSequence);
 
@@ -204,11 +211,8 @@ void UTakeRecorderWorldSource::AutotrackActors(class ULevelSequence* InSequence,
 					if (CanRecordActor(Actor))
 					{
 						UTakeRecorderActorSource* ActorSource = NewObject<UTakeRecorderActorSource>(Sources, UTakeRecorderActorSource::StaticClass(), NAME_None, RF_Transactional);
-						ActorSource->Target = Actor;
-
-						// Force the update to the recorded property map after change the actor source target. 
-						FPropertyChangedEvent PropertyChangedEvent(UTakeRecorderActorSource::StaticClass()->FindPropertyByName(GET_MEMBER_NAME_CHECKED(UTakeRecorderActorSource, Target)), EPropertyChangeType::ValueSet);
-						ActorSource->PostEditChangeProperty(PropertyChangedEvent);
+						// Setting the source actor through the function ensures internal property maps are also updated.
+						ActorSource->SetSourceActor(Actor);
 
 						// Add the source to our tracking table.
 						AutotrackedSources.Add(ActorSource);
diff --git a/Engine/Plugins/VirtualProduction/Takes/Source/TakeRecorderSources/Public/TakeRecorderActorSource.h b/Engine/Plugins/VirtualProduction/Takes/Source/TakeRecorderSources/Public/TakeRecorderActorSource.h
--- a/Engine/Plugins/VirtualProduction/Takes/Source/TakeRecorderSources/Public/TakeRecorderActorSource.h
+++ b/Engine/Plugins/VirtualProduction/Takes/Source/TakeRecorderSources/Public/TakeRecorderActorSource.h
@@ -303,11 +303,13 @@ protected:
 	
 	
 public:
+#if WITH_EDITOR
 	// UObject Interface
 	UE_API virtual void PostEditChangeProperty(struct FPropertyChangedEvent& PropertyChangedEvent) override;
 	UE_API virtual void PostDuplicate(bool bDuplicateForPIE) override;
 	UE_API virtual void PostEditUndo() override;
 	// ~UObject Interface
+#endif // WITH_EDITOR
 private:
 
 	UE_API virtual const FSlateBrush* GetDisplayIconImpl() const override;
diff --git a/Engine/Plugins/VirtualProduction/Takes/Source/TakeRecorderSources/Public/TakeRecorderLevelVisibilitySource.h b/Engine/Plugins/VirtualProduction/Takes/Source/TakeRecorderSources/Public/TakeRecorderLevelVisibilitySource.h
--- a/Engine/Plugins/VirtualProduction/Takes/Source/TakeRecorderSources/Public/TakeRecorderLevelVisibilitySource.h
+++ b/Engine/Plugins/VirtualProduction/Takes/Source/TakeRecorderSources/Public/TakeRecorderLevelVisibilitySource.h
@@ -18,7 +18,9 @@ public:
 
 	UE_API UTakeRecorderLevelVisibilitySourceSettings(const FObjectInitializer& ObjInit);
 
+#if WITH_EDITOR
 	UE_API virtual void PostEditChangeProperty(FPropertyChangedEvent& PropertyChangedEvent) override;
+#endif // WITH_EDITOR
 
 	// UTakeRecorderSource Interface
 	UE_API virtual FString GetSubsceneTrackName(ULevelSequence* InSequence) const override;
diff --git a/Engine/Plugins/VirtualProduction/Takes/Source/TakeRecorderSources/Public/TakeRecorderMicrophoneAudioManager.h b/Engine/Plugins/VirtualProduction/Takes/Source/TakeRecorderSources/Public/TakeRecorderMicrophoneAudioManager.h
--- a/Engine/Plugins/VirtualProduction/Takes/Source/TakeRecorderSources/Public/TakeRecorderMicrophoneAudioManager.h
+++ b/Engine/Plugins/VirtualProduction/Takes/Source/TakeRecorderSources/Public/TakeRecorderMicrophoneAudioManager.h
@@ -8,7 +8,10 @@
 
 #define UE_API TAKERECORDERSOURCES_API
 
+#if WITH_EDITOR
 class IAudioCaptureEditor;
+#endif // WITH_EDITOR
+
 class USoundWave;
 struct FTakeRecorderAudioSourceSettings;
 
@@ -28,7 +31,9 @@ public:
 	UE_API UTakeRecorderMicrophoneAudioManager(const FObjectInitializer& ObjInit);
 
 	// Begin UObject Interface
+#if WITH_EDITOR
 	UE_API virtual void PostEditChangeProperty(FPropertyChangedEvent& PropertyChangedEvent) override;
+#endif // WITH_EDITOR
 	// End UObject Interface
 
 	/** Enumerates the audio devices present on the current machine */
@@ -57,9 +62,11 @@ public:
 	 */
 	UE_API void FinalizeRecording();
 
-
+#if WITH_EDITOR
 	/** Fetches the USoundWave for this source after a Take has been recorded */
 	UE_API TObjectPtr<class USoundWave> GetRecordedSoundWave(const FTakeRecorderAudioSourceSettings& InSourceSettings);
+#endif // WITH_EDITOR
+
 	/** Accessor for the OnNotifySourcesOfDeviceChange delegate list */
 	FOnNotifySourcesOfDeviceChange& GetOnNotifySourcesOfDeviceChange() { return OnNotifySourcesOfDeviceChange; }
 
@@ -67,16 +74,22 @@ private:
 
 	/** Multicast delegate which notifies clients when the currently selected audio device changes. */
 	FOnNotifySourcesOfDeviceChange OnNotifySourcesOfDeviceChange;
+
+#if WITH_EDITOR
 	/** Calls factory to create the AudioRecorder object */
 	UE_API TUniquePtr<IAudioCaptureEditor> CreateAudioRecorderObject();
+#endif // WITH_EDITOR
+
 	/** Builds the list of audio devices which will be used in the device menu */
 	UE_API void BuildDeviceInfoArray();
 
 	/** Returns whether audio device with given Id was found during enumeration */
 	UE_API bool IsAudioDeviceAvailable(const FString& InDeviceId);
 
+#if WITH_EDITOR
 	/** The audio recorder object which manages low level recording of audio data */
 	TUniquePtr<IAudioCaptureEditor> AudioRecorder;
+#endif // WITH_EDITOR
 };
 
 #undef UE_API
diff --git a/Engine/Plugins/VirtualProduction/Takes/Source/TakeRecorderSources/Public/TakeRecorderMicrophoneAudioSource.h b/Engine/Plugins/VirtualProduction/Takes/Source/TakeRecorderSources/Public/TakeRecorderMicrophoneAudioSource.h
--- a/Engine/Plugins/VirtualProduction/Takes/Source/TakeRecorderSources/Public/TakeRecorderMicrophoneAudioSource.h
+++ b/Engine/Plugins/VirtualProduction/Takes/Source/TakeRecorderSources/Public/TakeRecorderMicrophoneAudioSource.h
@@ -28,7 +28,9 @@ public:
 
 	UE_API UTakeRecorderMicrophoneAudioSourceSettings(const FObjectInitializer& ObjInit);
 
+#if WITH_EDITOR
 	UE_API virtual void PostEditChangeProperty(FPropertyChangedEvent& PropertyChangedEvent) override;
+#endif // WITH_EDITOR
 	
 	// UTakeRecorderSource Interface
 	UE_API virtual FString GetSubsceneTrackName(ULevelSequence* InSequence) const override;
@@ -81,9 +83,11 @@ public:
 
 	UE_API UTakeRecorderMicrophoneAudioSource(const FObjectInitializer& ObjInit);
 
+#if WITH_EDITOR
 	// Begin UObject Interface
 	UE_API virtual void PostEditChangeProperty(FPropertyChangedEvent& PropertyChangedEvent) override;
 	// End UObject Interface
+#endif // WITH_EDITOR
 
 	/** Gain in decibels to apply to recorded audio */
 	UPROPERTY(EditAnywhere, BlueprintReadWrite, Category = "Source", meta = (ClampMin = "0.0", UIMin = "0.0", UIMax = "20.0", ClampMax = "40.0"))
diff --git a/Engine/Plugins/VirtualProduction/Takes/Source/TakeRecorderSources/Public/TakeRecorderWorldSource.h b/Engine/Plugins/VirtualProduction/Takes/Source/TakeRecorderSources/Public/TakeRecorderWorldSource.h
--- a/Engine/Plugins/VirtualProduction/Takes/Source/TakeRecorderSources/Public/TakeRecorderWorldSource.h
+++ b/Engine/Plugins/VirtualProduction/Takes/Source/TakeRecorderSources/Public/TakeRecorderWorldSource.h
@@ -23,7 +23,9 @@ public:
 
 	UE_API UTakeRecorderWorldSourceSettings(const FObjectInitializer& ObjInit);
 
+#if WITH_EDITOR
 	UE_API virtual void PostEditChangeProperty(FPropertyChangedEvent& PropertyChangedEvent) override;
+#endif // WITH_EDITOR
 
 	/** Record world settings */
 	UPROPERTY(config, EditAnywhere, BlueprintReadWrite, Category = "Source")
diff --git a/Engine/Plugins/VirtualProduction/Takes/Source/TakeRecorderSources/TakeRecorderSources.Build.cs b/Engine/Plugins/VirtualProduction/Takes/Source/TakeRecorderSources/TakeRecorderSources.Build.cs
--- a/Engine/Plugins/VirtualProduction/Takes/Source/TakeRecorderSources/TakeRecorderSources.Build.cs
+++ b/Engine/Plugins/VirtualProduction/Takes/Source/TakeRecorderSources/TakeRecorderSources.Build.cs
@@ -11,25 +11,17 @@ public class TakeRecorderSources : ModuleRules
                 "CinematicCamera",
                 "Core",
 				"CoreUObject",
-				"EditorFramework",
 				"Engine",
-				"AudioCaptureEditor",
 				"InputCore",
-				"LevelEditor",
 				"LevelSequence",
-                "LevelSequenceEditor",
                 "MovieScene",
 				"MovieSceneTracks",
-                "SceneOutliner",
-				"SequenceRecorder", // For FTimecodeBoneMethod
 				"SerializedRecorderInterface",
                 "Slate",
 				"SlateCore",
 				"TakesCore",
 				"TakeRecorder",
-				"TakeMovieScene",
-				"UnrealEd",
-				"Sequencer"
+				"TakeMovieScene"
 			}
 		);
 
@@ -40,5 +32,21 @@ public class TakeRecorderSources : ModuleRules
                 "Core",
             }
         );
+
+		if (Target.bBuildEditor == true)
+		{
+			PrivateDependencyModuleNames.AddRange(
+				new string[] {
+					"EditorFramework",
+					"AudioCaptureEditor",
+					"LevelEditor",
+					"LevelSequenceEditor",
+					"SceneOutliner",
+					"SequenceRecorder", // For FTimecodeBoneMethod
+					"UnrealEd",
+					"Sequencer",
+				}
+			);
+		}
 	}
 }
