Take Recorder - Split TakeRecorder into TakeRecorder and TakeRecorderEditor.

Adding a new TakeRecorderEditor module to accommodate all private TakeRecorder editor code into its own module that can be editor or uncooked. Everything left in the TakeRecorder module then is runtime compliant to allow for moving the TakeRecorder module to runtime.

#Synced-CL 52674206

diff --git a/Engine/Plugins/VirtualProduction/Takes/Source/TakeRecorder/Private/Recorder/TakeRecorder.cpp b/Engine/Plugins/VirtualProduction/Takes/Source/TakeRecorder/Private/Recorder/TakeRecorder.cpp
--- a/Engine/Plugins/VirtualProduction/Takes/Source/TakeRecorder/Private/Recorder/TakeRecorder.cpp
+++ b/Engine/Plugins/VirtualProduction/Takes/Source/TakeRecorder/Private/Recorder/TakeRecorder.cpp
@@ -3,7 +3,6 @@
 #include "Recorder/TakeRecorder.h"
 
 #include "Algo/Accumulate.h"
-#include "AnimatedRange.h"
 #include "AssetRegistry/AssetRegistryModule.h"
 #include "CoreGlobals.h"
 #include "Engine/Engine.h"
@@ -12,13 +11,11 @@
 #include "EntitySystem/MovieSceneInstanceRegistry.h"
 #include "Evaluation/MovieSceneEvaluationTemplateInstance.h"
 #include "AssetRegistry/IAssetRegistry.h"
-#include "ISequencer.h"
 #include "LevelSequence.h"
 #include "MetaData/MovieSceneShotMetaData.h"
 #include "Misc/CommandLine.h"
 #include "Misc/Parse.h"
 #include "MovieSceneTimeHelpers.h"
-#include "ObjectTools.h"
 #include "Recorder/TakeRecorderBlueprintLibrary.h"
 #include "Recorder/TakeRecorderSubsystem.h"
 #include "TakeMetaData.h"
@@ -29,15 +26,30 @@
 #include "TakesUtils.h"
 #include "Tickable.h"
 #include "Stats/Stats.h"
-#include "SequencerSettings.h"
 #include "ITakeRecorderSourcesManager.h"
 
 // Engine includes
 #include "GameFramework/WorldSettings.h"
 
+#if WITH_EDITOR
+#include "ISequencer.h"
+#include "ObjectTools.h"
+
+#include "AnimatedRange.h"
+
+#include "SequencerSettings.h"
+
 // UnrealEd includes
 #include "Editor.h"
 
+// LevelEditor includes
+#include "IAssetViewport.h"
+#include "LevelEditor.h"
+#include "MovieSceneToolHelpers.h"
+#include "Subsystems/AssetEditorSubsystem.h"
+
+#endif // WITH_EDITOR
+
 #include "UObject/GCObjectScopeGuard.h"
 
 // Slate includes
@@ -52,12 +64,7 @@
 #include "Widgets/Input/SButton.h"
 #include "Widgets/Layout/SBox.h"
 
-// LevelEditor includes
-#include "IAssetViewport.h"
-#include "LevelEditor.h"
-#include "MovieSceneToolHelpers.h"
 #include "Misc/GuardCVar.h"
-#include "Subsystems/AssetEditorSubsystem.h"
 
 #include UE_INLINE_GENERATED_CPP_BY_NAME(TakeRecorder)
 
@@ -82,12 +89,15 @@ public:
 		WeakFinishedAsset = InFinishedAsset;
 		TakeRecorderState = InTakeRecorder->GetState();
 
+	#if WITH_EDITOR
 		UTakeMetaData* TakeMetaData = InTakeRecorder->GetSequence()->FindMetaData<UTakeMetaData>();
 		check(TakeMetaData);
+	#endif // WITH_EDITOR
 
 		ChildSlot
 		[
 			SNew(SBox)
+		#if WITH_EDITOR
 			.Padding(FMargin(15.0f))
 			[
 				SNew(SVerticalBox)
@@ -152,6 +162,7 @@ public:
 					]
 				]
 			]
+		#endif // WITH_EDITOR
 		];
 	}
 
@@ -242,6 +253,7 @@ private:
 
 			ULevelSequence* LevelSequence = Recorder->GetSequence();
 
+		#if WITH_EDITOR
 			UTakeMetaData* TakeMetaData = LevelSequence ? LevelSequence->FindMetaData<UTakeMetaData>() : nullptr;
 
 			if (TakeMetaData)
@@ -255,6 +267,7 @@ private:
 
 				return FText::Format(LOCTEXT("RecordingTimecodeText", "Recording...{0}"), FText::FromString(Timecode.ToString()));
 			}
+		#endif // WITH_EDITOR
 		}
 
 		return LOCTEXT("RecordingText", "Recording...");
@@ -292,11 +305,13 @@ private:
 			Asset = Recorder ? Recorder->GetSequence() : nullptr;
 		}
 
+#if WITH_EDITOR
 		if (Asset)
 		{
 			TArray<FAssetData> Assets{ Asset };
 			GEditor->SyncBrowserToObjects(Assets);
 		}
+#endif // WITH_EDITOR
 	}
 
 	EVisibility CanBrowseToAssetFolder() const
@@ -516,7 +531,9 @@ bool UTakeRecorder::Initialize ( ULevelSequence* LevelSequenceBase, UTakeRecorde
 	if (Parameters.TakeRecorderMode == ETakeRecorderMode::RecordIntoSequence)
 	{
 		Parameters.Project.bStartAtCurrentTimecode = false;
+#if WITH_EDITOR
 		Parameters.User.bStopAtPlaybackEnd = true;
+#endif // WITH_EDITOR
 		Parameters.User.bAutoLock = false;
 	}
 
@@ -548,6 +565,7 @@ bool UTakeRecorder::Initialize ( ULevelSequence* LevelSequenceBase, UTakeRecorde
 
 	ensure(SetActiveRecorder(this));
 
+#if WITH_EDITOR
 	TSharedPtr<ISequencer> Sequencer = WeakSequencer.Pin();
 	if (Sequencer.IsValid())
 	{
@@ -582,6 +600,7 @@ bool UTakeRecorder::Initialize ( ULevelSequence* LevelSequenceBase, UTakeRecorde
 		Sequencer->SetViewRange(NewRange, EViewRangeInterpolation::Immediate);
 		Sequencer->SetClampRange(Sequencer->GetViewRange());
 	}
+#endif // WITH_EDITOR
 
 	return true;
 }
@@ -608,13 +627,18 @@ void UTakeRecorder::DiscoverSourceWorld()
 				OverlayWidget->AddToViewport();
 			}
 		}
+	#if WITH_EDITOR
 		FEditorDelegates::EndPIE.AddUObject(this, &UTakeRecorder::HandlePIE);
+	#endif // WITH_EDITOR
 	}
+
+#if WITH_EDITOR
 	// If not recording via PIE, be sure to stop recording if PIE Starts
 	if ( !bPlayInGame )
 	{
 		FEditorDelegates::BeginPIE.AddUObject(this, &UTakeRecorder::HandlePIE);//reuse same function
 	}
+#endif // WITH_EDITOR
 }
 
 bool UTakeRecorder::CreateDestinationAsset(const TCHAR* AssetPathFormat, ULevelSequence* LevelSequenceBase, UTakeRecorderSources* Sources, UTakeMetaData* MetaData, FText* OutError, bool bOpenSequencer)
@@ -634,11 +658,18 @@ bool UTakeRecorder::CreateDestinationAsset(const TCHAR* AssetPathFormat, ULevelS
 		return false;
 	}
 
+#if WITH_EDITOR
 	// Copy the sources into the level sequence for future reference (and potentially mutation throughout recording)
 	SequenceAsset->CopyMetaData(Sources);
+#endif // WITH_EDITOR
 
 	UMovieScene*   MovieScene    = SequenceAsset->GetMovieScene();
+
+#if WITH_EDITOR
 	UTakeMetaData* AssetMetaData = SequenceAsset->CopyMetaData(MetaData);
+#else // WITH_EDITOR
+	UTakeMetaData* AssetMetaData = DuplicateObject(MetaData, SequenceAsset);
+#endif // WITH_EDITOR
 
 	// Ensure the asset meta-data is unlocked for the recording (it is later Locked when the recording finishes)
 	AssetMetaData->Unlock();
@@ -659,11 +690,13 @@ bool UTakeRecorder::CreateDestinationAsset(const TCHAR* AssetPathFormat, ULevelS
 	FAssetRegistryModule::AssetCreated(SequenceAsset);
 
 	bool bSuccess = true;
+#if WITH_EDITOR
 	if (bOpenSequencer && GIsEditor)
 	{
 		WeakSequencer = TakesUtils::OpenSequencer(SequenceAsset, OutError);
 		bSuccess = WeakSequencer.IsValid();
 	}
+#endif // WITH_EDITOR
 	return bSuccess;
 }
 
@@ -671,6 +704,9 @@ bool UTakeRecorder::SetupDestinationAsset(const FTakeRecorderParameters& InParam
 {
 	check(LevelSequenceBase && Sources && MetaData);
 
+	SequenceAsset = LevelSequenceBase;
+
+#if WITH_EDITOR
 	if (bOpenSequencer)
 	{
 		WeakSequencer = TakesUtils::OpenSequencer(LevelSequenceBase, OutError);
@@ -681,15 +717,26 @@ bool UTakeRecorder::SetupDestinationAsset(const FTakeRecorderParameters& InParam
 	}
 
 	// The SequenceAsset is either LevelSequenceBase or the currently focused sequence
-	SequenceAsset = WeakSequencer.IsValid() ? Cast<ULevelSequence>(WeakSequencer.Pin()->GetFocusedMovieSceneSequence())
-		: LevelSequenceBase;
+	if (TSharedPtr<ISequencer> PinnedSequencer = WeakSequencer.Pin())
+	{
+		SequenceAsset = Cast<ULevelSequence>(PinnedSequencer->GetFocusedMovieSceneSequence());
+	}
+#endif // WITH_EDITOR
+
 	check(SequenceAsset);
 	
+#if WITH_EDITOR
 	// Copy the sources into the level sequence for future reference (and potentially mutation throughout recording)
 	SequenceAsset->CopyMetaData(Sources);
+#endif // WITH_EDITOR
 
 	UMovieScene*   MovieScene    = SequenceAsset->GetMovieScene();
+
+#if WITH_EDITOR
 	UTakeMetaData* AssetMetaData = SequenceAsset->CopyMetaData(MetaData);
+#else // WITH_EDITOR
+	UTakeMetaData* AssetMetaData = DuplicateObject(MetaData, SequenceAsset);
+#endif // WITH_EDITOR
 
 	// Ensure the asset meta-data is unlocked for the recording (it is later Locked when the recording finishes)
 	AssetMetaData->Unlock();
@@ -717,6 +764,7 @@ void UTakeRecorder::InitializeFromParameters()
 	// Set the end recording frame if enabled
 	StopRecordingFrame = Parameters.User.bStopAtPlaybackEnd ? SequenceAsset->GetMovieScene()->GetPlaybackRange().GetUpperBoundValue() : TOptional<FFrameNumber>();
 
+#if WITH_EDITOR
 	// Apply immersive mode if the parameters demand it
 	if (Parameters.User.bMaximizeViewport)
 	{
@@ -738,10 +786,12 @@ void UTakeRecorder::InitializeFromParameters()
 			OnStopCleanup.Add(RestoreImmersiveMode);
 		}
 	}
+#endif // WITH_EDITOR
 }
 
 bool UTakeRecorder::ShouldShowNotifications()
 {
+#if WITH_EDITOR
 	// -TAKERECORDERISHEADLESS in the command line can force headless behavior and disable the notifications.
 	static const bool bCmdLineTakeRecorderIsHeadless = FParse::Param(FCommandLine::Get(), TEXT("TAKERECORDERISHEADLESS"));
 
@@ -749,6 +799,9 @@ bool UTakeRecorder::ShouldShowNotifications()
 		&& !bCmdLineTakeRecorderIsHeadless
 		&& !FApp::IsUnattended()
 		&& !GIsRunningUnattendedScript;
+#else // WITH_EDITOR
+	return false;
+#endif // WITH_EDITOR
 }
 
 UWorld* UTakeRecorder::GetWorld() const
@@ -790,7 +843,10 @@ void UTakeRecorder::Tick(float DeltaTime)
 FQualifiedFrameTime UTakeRecorder::GetRecordTime() const
 {
 	float TimeDilation = 1.f;
+
+#if WITH_EDITOR
 	if (!WeakSequencer.IsValid())
+#endif // WITH_EDITOR
 	{
 		const UWorld* RecordingWorld = GetWorld();
 		check(RecordingWorld);
@@ -802,8 +858,13 @@ FQualifiedFrameTime UTakeRecorder::GetRecordTime() const
 		}
 	}
 	
+#if WITH_EDITOR
 	return TakesUtils::GetRecordTime(WeakSequencer.Pin(), SequenceAsset, TimecodeAtStart,
 		Parameters.Project.bStartAtCurrentTimecode, TimeDilation);
+#else // WITH_EDITOR
+	return TakesUtils::GetRecordTime(SequenceAsset, TimecodeAtStart,
+		Parameters.Project.bStartAtCurrentTimecode, TimeDilation);
+#endif // WITH_EDITOR
 }
 
 namespace UE::TakeRecorder
@@ -850,11 +911,11 @@ void UTakeRecorder::InternalTick(float DeltaTime)
 
 	FScopedSignedObjectModifyDefer FlushOnTick(true);
 
-	TSharedPtr<ISequencer> Sequencer = WeakSequencer.Pin();
 
 	FQualifiedFrameTime RecordTime = GetRecordTime();
 	
 	UTakeRecorderSources* Sources = ITakeRecorderSourcesManager::GetChecked().FindOrAddSources(SequenceAsset);
+
 	if (!Parameters.bDisableRecordingAndSave)
 	{
 		CurrentFrameTime = Sources->TickRecording(SequenceAsset, RecordTime, DeltaTime);
@@ -864,6 +925,8 @@ void UTakeRecorder::InternalTick(float DeltaTime)
 		CurrentFrameTime = Sources->AdvanceTime(RecordTime, DeltaTime);
 	}
 
+#if WITH_EDITOR
+	TSharedPtr<ISequencer> Sequencer = WeakSequencer.Pin();
 	if (Sequencer.IsValid())
 	{
 		FAnimatedRange Range = Sequencer->GetViewRange();
@@ -878,6 +941,7 @@ void UTakeRecorder::InternalTick(float DeltaTime)
 			Sequencer->SetClampRange(Sequencer->GetViewRange());
 		}
 	}
+#endif // WITH_EDITOR
 
 	OnTickRecordingEvent.Broadcast(this, RecordTime);
 
@@ -895,10 +959,15 @@ void UTakeRecorder::PreRecord()
 
 	UTakeRecorderSources* Sources = ITakeRecorderSourcesManager::GetChecked().FindSources(SequenceAsset);
 	check(Sources);
+
+#if WITH_EDITOR
 	UTakeMetaData* AssetMetaData = SequenceAsset->FindMetaData<UTakeMetaData>();
+#else // WITH_EDITOR
+	// Dummy out just a new UTakeMetaData object as it will not persist from a runtime environment anyway.
+	UTakeMetaData* AssetMetaData = NewObject<UTakeMetaData>();
+#endif // WITH_EDITOR
 
 	//Set the flag to specify if we should auto save the serialized data or not when recording.
-
 	MovieSceneSerializationNamespace::bAutoSerialize = Parameters.User.bAutoSerialize;
 	if (Parameters.User.bAutoSerialize)
 	{
@@ -931,17 +1000,24 @@ void UTakeRecorder::PreRecord()
 	TakeRecorderSourcesSettings.bStartAtCurrentTimecode = Parameters.Project.bStartAtCurrentTimecode;
 	TakeRecorderSourcesSettings.bRecordSourcesIntoSubSequences = Parameters.Project.bRecordSourcesIntoSubSequences;
 	TakeRecorderSourcesSettings.bRecordToPossessable = Parameters.Project.bRecordToPossessable;
-	TakeRecorderSourcesSettings.SpawnableOverwriteBehavior = Parameters.Project.SpawnableOverwriteBehavior;
+
+#if WITH_EDITOR
 	TakeRecorderSourcesSettings.bSaveRecordedAssets = (!Parameters.bDisableRecordingAndSave && Parameters.User.bSaveRecordedAssets) || GEditor == nullptr;
+#else // WITH_EDITOR
+	TakeRecorderSourcesSettings.bSaveRecordedAssets = (!Parameters.bDisableRecordingAndSave && Parameters.User.bSaveRecordedAssets);
+#endif // WITH_EDITOR
+
 	TakeRecorderSourcesSettings.bRemoveRedundantTracks = Parameters.User.bRemoveRedundantTracks;
 	TakeRecorderSourcesSettings.bAutoLock = Parameters.User.bAutoLock;
 	TakeRecorderSourcesSettings.SubSequenceDirectory = Parameters.Project.SubSequenceDirectory;
 
+#if WITH_EDITOR
 	TSharedPtr<ISequencer> Sequencer = WeakSequencer.Pin();
 	if (Sequencer)
 	{
 		TakeRecorderSourcesSettings.TransformOrigin = MovieSceneToolHelpers::GetTransformOriginForFocusedSequence(Sequencer);
 	}
+#endif // WITH_EDITOR
 
 	Sources->SetSettings(TakeRecorderSourcesSettings);
 
@@ -965,11 +1041,13 @@ void UTakeRecorder::PreRecord()
 		Sources->SetCachedAssets(SequenceAsset, Parameters.User.bAutoSerialize ? &ManifestSerializer : nullptr);
 	}
 
+#if WITH_EDITOR
 	// Refresh sequencer in case the movie scene data has mutated (ie. existing object bindings removed because they will be recorded again)
 	if (Sequencer.IsValid())
 	{
 		Sequencer->RefreshTree();
 	}
+#endif // WITH_EDITOR
 
 	// Apply engine Time Dilation after the countdown, otherwise the countdown will be dilated as well!
 	UWorld* RecordingWorld = GetWorld();
@@ -1011,16 +1089,18 @@ void UTakeRecorder::Start()
 
 	State = ETakeRecorderState::Started;
 
-	TSharedPtr<ISequencer> Sequencer = WeakSequencer.Pin();
 
 	CurrentFrameTime = FFrameTime(0);
 	TimecodeAtStart = Timecode;
 
+#if WITH_EDITOR
+	TSharedPtr<ISequencer> Sequencer = WeakSequencer.Pin();
 	// Discard any entity tokens we have so that restore state does not take effect when we delete any sections that recording will be replacing.
 	if (Sequencer.IsValid())
 	{
 		Sequencer->PreAnimatedState.DiscardEntityTokens();
 	}
+#endif // WITH_EDITOR
 
 	UMovieScene* MovieScene = SequenceAsset->GetMovieScene();
 	if (MovieScene)
@@ -1028,10 +1108,13 @@ void UTakeRecorder::Start()
 		CachedPlaybackRange = MovieScene->GetPlaybackRange();
 		CachedClockSource = MovieScene->GetClockSource();
 		MovieScene->SetClockSource(Parameters.Project.RecordingClockSource);
+
+#if WITH_EDITOR
 		if (Sequencer.IsValid())
 		{
 			Sequencer->ResetTimeController();
 		}
+#endif // WITH_EDITOR
 
 		FFrameRate FrameRate = MovieScene->GetDisplayRate();
 		FFrameRate TickResolution = MovieScene->GetTickResolution();
@@ -1056,16 +1139,25 @@ void UTakeRecorder::Start()
 
 		// Set infinite playback range when starting recording. Playback range will be clamped to the bounds of the sections at the completion of the recording
 		MovieScene->SetPlaybackRange(TRange<FFrameNumber>(PlaybackStartFrame, TNumericLimits<int32>::Max() - 1), false);
+
+#if WITH_EDITOR
 		if (Sequencer.IsValid())
 		{
 			Sequencer->SetGlobalTime(PlaybackStartFrame);
 			Sequencer->SetPlaybackStatus(EMovieScenePlayerStatus::Playing);
 		}
+#endif // WITH_EDITOR
 	}
+
 	UTakeRecorderSources* Sources = ITakeRecorderSourcesManager::GetChecked().FindSources(SequenceAsset);
 	check(Sources);
 
+#if WITH_EDITOR
 	UTakeMetaData* AssetMetaData = SequenceAsset->FindMetaData<UTakeMetaData>();
+#else // WITH_EDITOR
+	UTakeMetaData* AssetMetaData = NewObject<UTakeMetaData>();
+#endif // WITH_EDITOR
+
 	FDateTime UtcNow = FDateTime::UtcNow();
 	AssetMetaData->SetTimestamp(UtcNow);
 	AssetMetaData->SetTimecodeIn(Timecode);
@@ -1093,10 +1185,12 @@ void UTakeRecorder::Start()
 
 	UTakeRecorderBlueprintLibrary::OnTakeRecorderStarted();
 
+#if WITH_EDITOR
 	if (Sequencer.IsValid())
 	{
 		CompileSuppression = MakeUnique<FScopedVolatilityManagerSuppression>(Sequencer->GetSharedPlaybackState().ToSharedPtr());
 	}
+#endif // WITH_EDITOR
 }
 
 void UTakeRecorder::Stop()
@@ -1125,6 +1219,7 @@ void UTakeRecorder::StopInternal(const bool bCancelled)
 
 	TGuardValue<bool> ReentrantGuard(bStoppedRecording, true);
 
+#if WITH_EDITOR
 	if (Parameters.TakeRecorderMode == ETakeRecorderMode::RecordNewSequence)
 	{
 		USequencerSettings* SequencerSettings = USequencerSettingsContainer::GetOrCreate<USequencerSettings>(TEXT("TakeRecorderSequenceEditor"));
@@ -1132,31 +1227,38 @@ void UTakeRecorder::StopInternal(const bool bCancelled)
 		SequencerSettings->SetAllowEditsMode(CachedAllowEditsMode);
 		SequencerSettings->SetAutoChangeMode(CachedAutoChangeMode);
 	}
+#endif // WITH_EDITOR
 
 	ManifestSerializer.Close();
 
+#if WITH_EDITOR
 	FEditorDelegates::EndPIE.RemoveAll(this);
 	FEditorDelegates::BeginPIE.RemoveAll(this);
+#endif // WITH_EDITOR
 	
 	const bool bDidEverStartRecording = State == ETakeRecorderState::Started;
 	const bool bRecordingFinished = !bCancelled && bDidEverStartRecording;
 	State = bRecordingFinished ? ETakeRecorderState::Stopped : ETakeRecorderState::Cancelled;
 
+#if WITH_EDITOR
 	TSharedPtr<ISequencer> Sequencer = WeakSequencer.Pin();
 	if (Sequencer.IsValid())
 	{
 		Sequencer->SetPlaybackStatus(EMovieScenePlayerStatus::Stopped);
 	}
+#endif // WITH_EDITOR
 
 	UMovieScene* MovieScene = SequenceAsset ? SequenceAsset->GetMovieScene() : nullptr;
 	if (MovieScene)
 	{
 		MovieScene->SetClockSource(CachedClockSource);
 
+#if WITH_EDITOR
 		if (Sequencer.IsValid())
 		{
 			Sequencer->ResetTimeController();
 		}
+#endif // WITH_EDITOR
 	}
 
 	CompileSuppression.Reset();
@@ -1194,10 +1296,12 @@ void UTakeRecorder::StopInternal(const bool bCancelled)
 			
 			MovieScene->SetPlaybackRange(TRange<FFrameNumber>(Range.GetLowerBoundValue(), PlaybackEndFrame));
 
+#if WITH_EDITOR
 			if (Sequencer)
 			{
 				Sequencer->ResetTimeController();
 			}
+#endif // WITH_EDITOR
 		}
 
 		if (!Parameters.bDisableRecordingAndSave)
@@ -1222,6 +1326,7 @@ void UTakeRecorder::StopInternal(const bool bCancelled)
 
 		if (bRecordingFinished)
 		{
+		#if WITH_EDITOR
 			// Lock the sequence so that it can't be changed without implicitly unlocking it now
 			if (Parameters.User.bAutoLock)
 			{
@@ -1231,8 +1336,13 @@ void UTakeRecorder::StopInternal(const bool bCancelled)
 			// We have recorded this sequence, mark it.
 			UMovieSceneShotMetaData* MovieSceneMetaData = SequenceAsset->FindOrAddMetaData<UMovieSceneShotMetaData>();
 			MovieSceneMetaData->SetIsRecorded(true);
+		#endif // WITH_EDITOR
 
+		#if WITH_EDITOR
 			UTakeMetaData* AssetMetaData = SequenceAsset->FindMetaData<UTakeMetaData>();
+		#else // WITH_EDITOR
+			UTakeMetaData* AssetMetaData = NewObject<UTakeMetaData>();
+		#endif // WITH_EDITOR
 			check(AssetMetaData);
 
 			if (MovieScene)
@@ -1244,16 +1354,19 @@ void UTakeRecorder::StopInternal(const bool bCancelled)
 				AssetMetaData->SetTimecodeOut(Timecode);
 			}
 
+#if WITH_EDITOR
 			if (GEditor && GEditor->GetEditorWorldContext().World())
 			{
 				AssetMetaData->SetLevelOrigin(GEditor->GetEditorWorldContext().World()->PersistentLevel);
 			}
+#endif // WITH_EDITOR
 
 			// Lock the meta data so it can't be changed without implicitly unlocking it now
 			AssetMetaData->Lock();
 
 			if (Parameters.User.bSaveRecordedAssets)
 			{
+#if WITH_EDITOR
 				if (Sequencer)
 				{
 					// Ensure we don't have any external references prior to saving. This automatically handles possessables vs spawnables.
@@ -1267,17 +1380,20 @@ void UTakeRecorder::StopInternal(const bool bCancelled)
 					Sequencer->Save();
 				}
 				else
+#endif // WITH_EDITOR
 				{
 					// Fallback to "normal" saving if sequencer is not available (seems like it should be though?). Asset won't have any thumbnail.
 					TakesUtils::SaveAsset(SequenceAsset);
 				}
 			}
 
+#if WITH_EDITOR
 			// Rebuild sequencer because subsequences could have been added or bindings removed
 			if (Sequencer)
 			{
 				Sequencer->NotifyMovieSceneDataChanged(EMovieSceneDataChangeType::MovieSceneStructureItemsChanged);
 			}
+#endif // WITH_EDITOR
 		}
 	}
 
@@ -1308,13 +1424,16 @@ void UTakeRecorder::StopInternal(const bool bCancelled)
 		}
 	}
 
+#if WITH_EDITOR
 	// Tear down the Sequencer instance before we delete the SequenceAsset.
 	Sequencer.Reset();
+#endif // WITH_EDITOR
 
 	// Delete the asset after OnTakeRecorderFinished and OnTakeRecorderCancelled because the ScopedSequencerPanel 
 	// will still have the current sequence before it is returned to the Pending Take.
 	if (Parameters.TakeRecorderMode == ETakeRecorderMode::RecordNewSequence && !bRecordingFinished)
 	{
+#if WITH_EDITOR
 		if (GIsEditor)
 		{
 			// Recording was canceled before it started, so delete the asset. Note we can only do this on editor
@@ -1323,6 +1442,7 @@ void UTakeRecorder::StopInternal(const bool bCancelled)
 			//
 			FAssetRegistryModule::AssetDeleted(SequenceAsset);
 		}
+#endif // WITH_EDITOR
 
 		// Move the asset to the transient package so that new takes with the same number can be created in its place
 		FName DeletedPackageName = MakeUniqueObjectName(nullptr, UPackage::StaticClass(), *(FString(TEXT("/Temp/") + SequenceAsset->GetName() + TEXT("_Cancelled"))));
@@ -1389,6 +1509,7 @@ FOnTickRecording& UTakeRecorder::OnTickRecording()
 	return OnTickRecordingEvent;
 }
 
+#if WITH_EDITOR
 void UTakeRecorder::HandlePIE(bool bIsSimulating)
 {
 	ULevelSequence* FinishedAsset = GetSequence();
@@ -1411,6 +1532,7 @@ void UTakeRecorder::HandlePIE(bool bIsSimulating)
 		GEditor->GetEditorSubsystem<UAssetEditorSubsystem>()->CloseAllEditorsForAsset(FinishedAsset);
 	}
 }
+#endif // WITH_EDITOR
 
 #undef LOCTEXT_NAMESPACE
 
diff --git a/Engine/Plugins/VirtualProduction/Takes/Source/TakeRecorder/Private/Recorder/TakeRecorderBlueprintLibrary.cpp b/Engine/Plugins/VirtualProduction/Takes/Source/TakeRecorder/Private/Recorder/TakeRecorderBlueprintLibrary.cpp
--- a/Engine/Plugins/VirtualProduction/Takes/Source/TakeRecorder/Private/Recorder/TakeRecorderBlueprintLibrary.cpp
+++ b/Engine/Plugins/VirtualProduction/Takes/Source/TakeRecorder/Private/Recorder/TakeRecorderBlueprintLibrary.cpp
@@ -3,13 +3,15 @@
 #include "Recorder/TakeRecorderBlueprintLibrary.h"
 #include "Recorder/TakeRecorder.h"
 #include "Recorder/TakeRecorderPanel.h"
-#include "Widgets/STakeRecorderTabContent.h"
 #include "ITakeRecorderModule.h"
 #include "TakeRecorderSettings.h"
 #include "UObject/Package.h"
 
 #include "Modules/ModuleManager.h"
+
+#if WITH_EDITOR
 #include "LevelEditor.h"
+#endif // WITH_EDITOR
 
 #include UE_INLINE_GENERATED_CPP_BY_NAME(TakeRecorderBlueprintLibrary)
 
@@ -126,12 +128,14 @@ UTakeRecorderPanel* UTakeRecorderBlueprintLibrary::OpenTakeRecorderPanel()
 		return Existing;
 	}
 
+#if WITH_EDITOR
 	if (!IsRunningGame())
 	{
 		FLevelEditorModule& LevelEditorModule = FModuleManager::LoadModuleChecked<FLevelEditorModule>("LevelEditor");
 		LevelEditorModule.GetLevelEditorTabManager()->TryInvokeTab(ITakeRecorderModule::TakeRecorderTabName);
 	}
 	else
+#endif // WITH_EDITOR
 	{
 		FFrame::KismetExecutionMessage(TEXT("The Take Recorder Panel will not open because the game is running."), ELogVerbosity::Log);
 	}
diff --git a/Engine/Plugins/VirtualProduction/Takes/Source/TakeRecorder/Private/Recorder/TakeRecorderPanel.cpp b/Engine/Plugins/VirtualProduction/Takes/Source/TakeRecorder/Private/Recorder/TakeRecorderPanel.cpp
--- a/Engine/Plugins/VirtualProduction/Takes/Source/TakeRecorder/Private/Recorder/TakeRecorderPanel.cpp
+++ b/Engine/Plugins/VirtualProduction/Takes/Source/TakeRecorder/Private/Recorder/TakeRecorderPanel.cpp
@@ -1,11 +1,7 @@
 // Copyright Epic Games, Inc. All Rights Reserved.
 
 #include "Recorder/TakeRecorderPanel.h"
-#include "Widgets/STakeRecorderTabContent.h"
-#include "TakePresetToolkit.h"
-
-#include "LevelEditor.h"
-#include "Modules/ModuleManager.h"
+#include "ITakeRecorderTabContent.h"
 
 #include UE_INLINE_GENERATED_CPP_BY_NAME(TakeRecorderPanel)
 
@@ -14,7 +10,7 @@ bool UTakeRecorderPanel::IsPanelOpen() const
 	return WeakTabContent.Pin().IsValid();
 }
 
-void UTakeRecorderPanel::InitializePanel(TWeakPtr<STakeRecorderTabContent> InTabContent)
+void UTakeRecorderPanel::InitializePanel(TWeakPtr<ITakeRecorderTabContent> InTabContent)
 {
 	WeakTabContent = InTabContent;
 }
@@ -73,10 +69,7 @@ void UTakeRecorderPanel::SetupForEditing(UTakePreset* TakePreset)
 {
 	if (ValidateTabContent())
 	{
-		TSharedPtr<ILevelEditor> OpenedFromLevelEditor = FModuleManager::LoadModuleChecked< FLevelEditorModule >("LevelEditor").GetFirstLevelEditor();
-		TSharedPtr<FTakePresetToolkit> Toolkit = MakeShared<FTakePresetToolkit>();
-		Toolkit->Initialize(EToolkitMode::WorldCentric, OpenedFromLevelEditor, TakePreset);
-		WeakTabContent.Pin()->SetupForEditing(Toolkit);
+		WeakTabContent.Pin()->SetupForEditing(TakePreset);
 	}
 }
 
diff --git a/Engine/Plugins/VirtualProduction/Takes/Source/TakeRecorder/Private/Recorder/TakeRecorderSubsystemImplementation.cpp b/Engine/Plugins/VirtualProduction/Takes/Source/TakeRecorder/Private/Recorder/TakeRecorderSubsystemImplementation.cpp
--- a/Engine/Plugins/VirtualProduction/Takes/Source/TakeRecorder/Private/Recorder/TakeRecorderSubsystemImplementation.cpp
+++ b/Engine/Plugins/VirtualProduction/Takes/Source/TakeRecorder/Private/Recorder/TakeRecorderSubsystemImplementation.cpp
@@ -5,29 +5,31 @@
 #include "AssetRegistry/AssetRegistryModule.h"
 #include "Dialog/SMessageDialog.h"
 #include "Framework/Notifications/NotificationManager.h"
-#include "ILevelSequenceEditorToolkit.h"
 #include "ITakeRecorderModule.h"
 #include "ITakeRecorderNamingTokensModule.h"
 #include "LevelSequence.h"
 #include "MovieScene.h"
-#include "NamingTokens.h"
-#include "NamingTokensEngineSubsystem.h"
 #include "Recorder/TakeRecorder.h"
 #include "Recorder/TakeRecorderBlueprintLibrary.h"
 #include "Recorder/TakeRecorderParameters.h"
 #include "Recorder/TakeRecorderSubsystem.h"
-#include "ScopedTransaction.h"
 #include "TakeMetaData.h"
 #include "TakePresetSettings.h"
-#include "TakeRecorderModule.h"
 #include "TakeRecorderSettings.h"
 #include "TakeRecorderSource.h"
 #include "TakeRecorderSources.h"
 #include "TakesCoreBlueprintLibrary.h"
 #include "Widgets/Notifications/SNotificationList.h"
-#include "Widgets/TakeRecorderNamingTokensData.h"
+#include "TakeRecorderNamingTokensData.h"
 #include "ITakeRecorderSourcesManager.h"
 
+#if WITH_EDITOR
+#include "ILevelSequenceEditorToolkit.h"
+#include "NamingTokensEngineSubsystem.h"
+#include "NamingTokens.h"
+#include "ScopedTransaction.h"
+#endif // WITH_EDITOR
+
 #include UE_INLINE_GENERATED_CPP_BY_NAME(TakeRecorderSubsystemImplementation)
 
 #define LOCTEXT_NAMESPACE "TakeRecorderSubsystemImplementation"
@@ -113,8 +115,10 @@ void UTakeRecorderSubsystemImplementation::DeinitializeImplementation()
 
 	UTakePresetSettings::Get()->OnSettingsChanged().RemoveAll(this);
 	
+#if WITH_EDITOR
 	FEditorDelegates::OnPreForceDeleteObjects.Remove(OnPreForceDeleteObjectsHandle);
 	OnPreForceDeleteObjectsHandle.Reset();
+#endif // WITH_EDITOR
 	
 	if (const FAssetRegistryModule* AssetRegistryModule = FModuleManager::GetModulePtr<FAssetRegistryModule>("AssetRegistry"))
 	{
@@ -124,9 +128,9 @@ void UTakeRecorderSubsystemImplementation::DeinitializeImplementation()
 		}
 	}
 
-	if (ITakeRecorderModule* TakeRecorderModule = FModuleManager::Get().GetModulePtr<ITakeRecorderModule>("TakeRecorder"))
+	if (ITakeRecorderModule::IsAvailable())
 	{
-		TakeRecorderModule->GetLastLevelSequenceProvider().Unbind();
+		ITakeRecorderModule::Get().GetLastLevelSequenceProvider().Unbind();
 	}
 
 	UnbindNamingTokensEvents();
@@ -140,8 +144,8 @@ void UTakeRecorderSubsystemImplementation::DeinitializeImplementation()
 void UTakeRecorderSubsystemImplementation::SetTargetSequence(const FTakeRecorderSequenceParameters& InData)
 {
 	TargetSequenceData = InData;
-	
-	ITakeRecorderModule& TakeRecorderModule = FModuleManager::Get().LoadModuleChecked<ITakeRecorderModule>("TakeRecorder");
+
+	ITakeRecorderModule& TakeRecorderModule = ITakeRecorderModule::Get();
 	TakeRecorderModule.GetLastLevelSequenceProvider().Unbind();
 	TakeRecorderModule.GetLastLevelSequenceProvider().BindUObject(this, &UTakeRecorderSubsystemImplementation::SetLastLevelSequence);
 	
@@ -172,23 +176,20 @@ void UTakeRecorderSubsystemImplementation::SetTargetSequence(const FTakeRecorder
 		TransientPreset->CopyFrom(InData.BaseSequence);
 
 		ULevelSequence* LevelSequence = TransientPreset->GetLevelSequence();
+		ITakeRecorderSourcesManager::GetChecked().CopySources(InData.BaseSequence, LevelSequence);
 
-		UTakeRecorderSources* BaseSources = ITakeRecorderSourcesManager::GetChecked().FindSources(InData.BaseSequence);
-		if (BaseSources && LevelSequence)
-		{
-			LevelSequence->CopyMetaData(BaseSources);
-		}
-
+	#if WITH_EDITOR
 		if (LevelSequence)
 		{
 			LevelSequence->GetMovieScene()->SetReadOnly(false);
 		}
 
-		if (UTakeMetaData*  TakeMetaDataLevelSequence = LevelSequence ? LevelSequence->FindMetaData<UTakeMetaData>() : nullptr)
+		if (UTakeMetaData* TakeMetaDataLevelSequence = LevelSequence ? LevelSequence->FindMetaData<UTakeMetaData>() : nullptr)
 		{
 			TakeMetaDataLevelSequence->Unlock();
 			TakeMetaDataLevelSequence->SetTimestamp(FDateTime(0));
 		}
+	#endif // WITH_EDITOR
 	}
 	else if (InData.RecordIntoSequence)
 	{
@@ -226,6 +227,7 @@ void UTakeRecorderSubsystemImplementation::SetRecordIntoLevelSequence(ULevelSequ
 	SuppliedLevelSequence = nullptr;
 	RecordIntoLevelSequence = LevelSequence;
 
+#if WITH_EDITOR
 	if (RecordIntoLevelSequence)
 	{
 		RecordIntoLevelSequence->GetMovieScene()->SetReadOnly(false);
@@ -241,6 +243,7 @@ void UTakeRecorderSubsystemImplementation::SetRecordIntoLevelSequence(ULevelSequ
 		LocalTakeMetaData->SetFrameRate(LevelSequence->GetMovieScene()->GetDisplayRate());
 		LocalTakeMetaData->SetFrameRateFromTimecode(false);
 	}
+#endif // WITH_EDITOR
 }
 
 bool UTakeRecorderSubsystemImplementation::CanReviewLastRecording() const
@@ -249,7 +252,7 @@ bool UTakeRecorderSubsystemImplementation::CanReviewLastRecording() const
 	{
 		return false;
 	}
-	ITakeRecorderModule& TakeRecorderModule = FModuleManager::Get().LoadModuleChecked<ITakeRecorderModule>("TakeRecorder");
+	ITakeRecorderModule& TakeRecorderModule = ITakeRecorderModule::Get();
 	FCanReviewLastRecordedLevelSequence& CanReview = TakeRecorderModule.GetCanReviewLastRecordedLevelSequenceDelegate();
 	if (CanReview.IsBound())
 	{
@@ -301,6 +304,7 @@ bool UTakeRecorderSubsystemImplementation::StartRecording(bool bOpenSequencer, b
 
 		FText ErrorText = LOCTEXT("UnknownError", "An unknown error occurred when trying to start recording");
 
+	#if WITH_EDITOR
 		if (GIsEditor)
 		{
 			IAssetEditorInstance* AssetEditor = GEditor->GetEditorSubsystem<UAssetEditorSubsystem>()->FindEditorForAsset(LevelSequence, false);
@@ -316,6 +320,7 @@ bool UTakeRecorderSubsystemImplementation::StartRecording(bool bOpenSequencer, b
 				}
 			}
 		}
+	#endif // WITH_EDITOR
 
 		UTakeRecorder* NewRecorder = NewObject<UTakeRecorder>(GetTransientPackage(), NAME_None, RF_Transient);
 
@@ -351,13 +356,14 @@ void UTakeRecorderSubsystemImplementation::CancelRecording(bool bShowConfirmMess
 	if (UTakeRecorder* CurrentRecording = UTakeRecorder::GetActiveRecorder())
 	{
 		auto DoCancel = [](const TWeakObjectPtr<UTakeRecorder>& WeakRecording)
-		{
-			if (WeakRecording.IsValid())
 			{
-				WeakRecording.Get()->Cancel(); 
-			}
-		};
+				if (WeakRecording.IsValid())
+				{
+					WeakRecording.Get()->Cancel();
+				}
+			};
 
+	#if WITH_EDITOR
 		if (bShowConfirmMessage)
 		{
 			const TSharedRef<SMessageDialog> ConfirmDialog = SNew(SMessageDialog)
@@ -375,6 +381,7 @@ void UTakeRecorderSubsystemImplementation::CancelRecording(bool bShowConfirmMess
 			ConfirmDialog->Show();
 		}
 		else
+	#endif // WITH_EDITOR
 		{
 			DoCancel(TWeakObjectPtr<UTakeRecorder>(CurrentRecording));
 		}
@@ -416,7 +423,9 @@ void UTakeRecorderSubsystemImplementation::ClearPendingTake()
 	SuppliedLevelSequence = nullptr;
 	RecordIntoLevelSequence = nullptr;
 
+#if WITH_EDITOR
 	FScopedTransaction Transaction(LOCTEXT("ClearPendingTake_Transaction", "Clear Pending Take"));
+#endif // WITH_EDITOR
 
 	TransientPreset->Modify();
 	TransientPreset->CreateLevelSequence();
@@ -430,7 +439,7 @@ void UTakeRecorderSubsystemImplementation::ClearPendingTake()
 
 UTakePreset* UTakeRecorderSubsystemImplementation::GetPendingTake() const
 {
-	const ITakeRecorderModule& TakeRecorderModule = FModuleManager::LoadModuleChecked<ITakeRecorderModule>("TakeRecorder");
+	const ITakeRecorderModule& TakeRecorderModule = ITakeRecorderModule::Get();
 	return TakeRecorderModule.GetPendingTake();
 }
 
@@ -438,7 +447,9 @@ void UTakeRecorderSubsystemImplementation::RevertChanges()
 {
 	UTakePreset* PresetOrigin = GetTakeMetaData()->GetPresetOrigin();
 
+#if WITH_EDITOR
 	FScopedTransaction Transaction(LOCTEXT("RevertChanges_Transaction", "Revert Changes"));
+#endif // WITH_EDITOR
 
 	TransientPreset->Modify();
 	TransientPreset->CopyFrom(PresetOrigin);
@@ -456,7 +467,9 @@ UTakeRecorderSource* UTakeRecorderSubsystemImplementation::AddSource(const TSubc
 
 	if (*InSourceClass && Sources)
 	{
+	#if WITH_EDITOR
 		FScopedTransaction Transaction(FText::Format(LOCTEXT("AddNewSource_Transaction", "Add New {0} Source"), InSourceClass->GetDisplayNameText()));
+	#endif // WITH_EDITOR
 		Sources->Modify();
 
 		return Sources->AddSource(InSourceClass);
@@ -470,7 +483,9 @@ void UTakeRecorderSubsystemImplementation::RemoveSource(UTakeRecorderSource* InS
 	UTakeRecorderSources* Sources = GetSources();
 	if (Sources && InSource)
 	{
+	#if WITH_EDITOR
 		FScopedTransaction Transaction(FText::Format(LOCTEXT("RemoveSource_Transaction", "Remove {0} Source"), InSource->GetDisplayText()));
+	#endif // WITH_EDITOR
 		Sources->Modify();
 		Sources->RemoveSource(InSource);
 	}
@@ -480,7 +495,9 @@ void UTakeRecorderSubsystemImplementation::ClearSources()
 {
 	if (UTakeRecorderSources* Sources = GetSources())
 	{
+	#if WITH_EDITOR
 		FScopedTransaction Transaction(LOCTEXT("ClearSources_Transaction", "Clear Sources"));
+	#endif // WITH_EDITOR
 		Sources->Modify();
 		for (UTakeRecorderSource* Source : Sources->GetSourcesCopy())
 		{
@@ -492,8 +509,7 @@ void UTakeRecorderSubsystemImplementation::ClearSources()
 UTakeRecorderSources* UTakeRecorderSubsystemImplementation::GetSources() const
 {
 	ULevelSequence* LevelSequence = GetLevelSequence();
-	UTakeRecorderSources* Sources = ITakeRecorderSourcesManager::GetChecked().FindOrAddSources(LevelSequence);
-	return Sources;
+	return ITakeRecorderSourcesManager::GetChecked().FindOrAddSources(LevelSequence);
 }
 
 TArrayView<UTakeRecorderSource* const> UTakeRecorderSubsystemImplementation::GetAllSources() const
@@ -529,7 +545,7 @@ void UTakeRecorderSubsystemImplementation::AddSourceForActor(AActor* InActor, bo
 {
 	if (UTakeRecorderSources* Sources = GetSources())
 	{
-		FTakeRecorderModule& TakeRecorderModule = FModuleManager::GetModuleChecked<FTakeRecorderModule>("TakeRecorder");
+		ITakeRecorderModule& TakeRecorderModule = ITakeRecorderModule::Get();
 		TakeRecorderModule.GetSourcesExtensionData().OnAddActorSource.ExecuteIfBound(Sources, { InActor }, bReduceKeys, bShowProgress);
 	}
 }
@@ -538,14 +554,14 @@ void UTakeRecorderSubsystemImplementation::RemoveActorFromSources(AActor* InActo
 {
 	if (UTakeRecorderSources* Sources = GetSources())
 	{
-		FTakeRecorderModule& TakeRecorderModule = FModuleManager::GetModuleChecked<FTakeRecorderModule>("TakeRecorder");
+		ITakeRecorderModule& TakeRecorderModule = ITakeRecorderModule::Get();
 		TakeRecorderModule.GetSourcesExtensionData().OnRemoveActorSource.ExecuteIfBound(Sources, { InActor });
 	}
 }
 
 AActor* UTakeRecorderSubsystemImplementation::GetSourceActor(UTakeRecorderSource* InSource) const
 {
-	FTakeRecorderModule& TakeRecorderModule = FModuleManager::GetModuleChecked<FTakeRecorderModule>("TakeRecorder");
+	ITakeRecorderModule& TakeRecorderModule = ITakeRecorderModule::Get();
 	if (TakeRecorderModule.GetSourcesExtensionData().OnGetSourceActor.IsBound())
 	{
 		return TakeRecorderModule.GetSourcesExtensionData().OnGetSourceActor.Execute(InSource);
@@ -568,7 +584,9 @@ void UTakeRecorderSubsystemImplementation::SetTakeNumber(int32 InNewTakeNumber,
 {
 	if (TakeMetaData)
 	{
+	#if WITH_EDITOR
 		FScopedTransaction Transaction(LOCTEXT("SetTake_Transaction", "Set Take Number"));
+	#endif // WITH_EDITOR
 		TakeMetaData->Modify();
 		TakeMetaData->SetTakeNumber(InNewTakeNumber, bEmitChanged);
 		bAutoApplyTakeNumber = false;
@@ -618,11 +636,13 @@ void UTakeRecorderSubsystemImplementation::IncrementTakeNumber()
 	
 	if (const ULevelSequence* LevelSequence = TransientPreset ? TransientPreset->GetLevelSequence() : nullptr)
 	{
+	#if WITH_EDITOR
 		if (UTakeMetaData* MetaData = LevelSequence->FindMetaData<UTakeMetaData>())
 		{
 			const int32 NextTakeNumber = GetNextTakeNumber(MetaData->GetSlate());
 			MetaData->SetTakeNumber(NextTakeNumber);
 		}
+	#endif // WITH_EDITOR
 	}
 
 	bAutoApplyTakeNumber = true;
@@ -652,7 +672,9 @@ void UTakeRecorderSubsystemImplementation::SetSlateName(const FString& InSlateNa
 {
 	if (TakeMetaData && TakeMetaData->GetSlate() != InSlateName)
 	{
+	#if WITH_EDITOR
 		FScopedTransaction Transaction(LOCTEXT("SetSlate_Transaction", "Set Take Slate"));
+	#endif // WITH_EDITOR
 		TakeMetaData->Modify();
 
 		TakeMetaData->SetSlate(InSlateName, bEmitChanged);
@@ -742,7 +764,9 @@ void UTakeRecorderSubsystemImplementation::ImportPreset(const FAssetData& InPres
 	RecordIntoLevelSequence = nullptr;
 
 	UTakePreset* Take = CastChecked<UTakePreset>(InPreset.GetAsset());
+#if WITH_EDITOR
 	FScopedTransaction Transaction(LOCTEXT("ImportPreset_Transaction", "Import Take Preset"));
+#endif // WITH_EDITOR
 
 	TransientPreset->Modify();
 	TransientPreset->CopyFrom(Take);
@@ -894,7 +918,9 @@ void UTakeRecorderSubsystemImplementation::CacheMetaData()
 
 	if (const ULevelSequence* LevelSequence = GetLevelSequence())
 	{
+	#if WITH_EDITOR
 		NewMetaDataThisTick = LevelSequence->FindMetaData<UTakeMetaData>();
+	#endif // WITH_EDITOR
 	}
 
 	// If it's null we use the transient meta-data
@@ -926,12 +952,14 @@ void UTakeRecorderSubsystemImplementation::CacheMetaData()
 	{
 		TakeMetaData = NewMetaDataThisTick;
 
+	#if WITH_EDITOR
 		if (!OnPreForceDeleteObjectsHandle.IsValid())
 		{
 			// Since this is a strong reference, we need to be able to clear it if the user is attempting to remove it, but let's not always
 			// hook this event since the subsystem won't be active for general UE operation.
 			OnPreForceDeleteObjectsHandle = FEditorDelegates::OnPreForceDeleteObjects.AddUObject(this, &UTakeRecorderSubsystemImplementation::OnPreForceDeleteObjects);
 		}
+	#endif // WITH_EDITOR
 	}
 
 	//Set MovieScene Display Rate to the Preset Frame Rate.
@@ -1085,6 +1113,7 @@ void UTakeRecorderSubsystemImplementation::OnSourceModified(UTakeRecorderSource*
 
 void UTakeRecorderSubsystemImplementation::BindToNamingTokenEvents()
 {
+#if WITH_EDITOR
 	if (GEngine && !GetNamingTokensData()->TakeRecorderNamingTokens.IsValid())
 	{
 		GetNamingTokensData()->TakeRecorderNamingTokens = GEngine->GetEngineSubsystem<UNamingTokensEngineSubsystem>()->GetNamingTokens(ITakeRecorderNamingTokensModule::GetTakeRecorderNamespace());
@@ -1095,17 +1124,21 @@ void UTakeRecorderSubsystemImplementation::BindToNamingTokenEvents()
 			TakeRecorderNamingTokensPin->GetOnPreEvaluateEvent().AddUObject(this, &UTakeRecorderSubsystemImplementation::OnTakeRecorderNamingTokensPreEvaluate);
 		}
 	}
+#endif // WITH_EDITOR
 }
 
 void UTakeRecorderSubsystemImplementation::UnbindNamingTokensEvents()
 {
+#if WITH_EDITOR
 	if (const TStrongObjectPtr<UNamingTokens> TakeRecorderNamingTokensPin = GetNamingTokensData()->TakeRecorderNamingTokens.Pin())
 	{
 		TakeRecorderNamingTokensPin->UnregisterExternalTokens(GetNamingTokensData()->NamingTokensExternalGuid);
 		TakeRecorderNamingTokensPin->GetOnPreEvaluateEvent().RemoveAll(this);
 	}
+#endif // WITH_EDITOR
 }
 
+#if WITH_EDITOR
 void UTakeRecorderSubsystemImplementation::OnTakeRecorderNamingTokensPreEvaluate(const FNamingTokensEvaluationData& InEvaluationData)
 {
 	if (GetNamingTokensData()->TakeRecorderNamingTokens.IsValid())
@@ -1126,16 +1159,19 @@ void UTakeRecorderSubsystemImplementation::OnTakeRecorderNamingTokensPreEvaluate
 		}
 	}
 }
+#endif // WITH_EDITOR
 
 void UTakeRecorderSubsystemImplementation::OnTakePresetSettingsChanged()
 {
 	// This effectively regenerates the level sequence in the preset.
 	// Existing data is discarded. In the future, we should try to migrate it.
 	// FTakePresetRecorderCustomization handles asking the user whether it is ok to discard the changes. Once this event fires, the decision has been made.
+#if WITH_EDITOR
 	if (!GIsTransacting)
 	{
 		ClearPendingTake();
 	}
+#endif // WITH_EDITOR
 }
 
 void UTakeRecorderSubsystemImplementation::OnPreForceDeleteObjects(const TArray<UObject*>& InObjects)
@@ -1149,7 +1185,10 @@ void UTakeRecorderSubsystemImplementation::OnPreForceDeleteObjects(const TArray<
 		LevelSequence = LastRecordedLevelSequence.Get();
 	}
 	if (LevelSequence && TakeMetaData && InObjects.Contains(LevelSequence)
-		&& LevelSequence->FindMetaData<UTakeMetaData>() == TakeMetaData)
+	#if WITH_EDITOR
+		&& LevelSequence->FindMetaData<UTakeMetaData>() == TakeMetaData
+	#endif // WITH_EDITOR
+		)
 	{
 		if (IsRecording())
 		{
@@ -1157,8 +1196,11 @@ void UTakeRecorderSubsystemImplementation::OnPreForceDeleteObjects(const TArray<
 			CancelRecording(false);
 		}
 		TakeMetaData = nullptr;
+
+	#if WITH_EDITOR
 		FEditorDelegates::OnPreForceDeleteObjects.Remove(OnPreForceDeleteObjectsHandle);
 		OnPreForceDeleteObjectsHandle.Reset();
+	#endif // WITH_EDITOR
 	}
 }
 
diff --git a/Engine/Plugins/VirtualProduction/Takes/Source/TakeRecorder/Private/SerializedRecorder.cpp b/Engine/Plugins/VirtualProduction/Takes/Source/TakeRecorder/Private/SerializedRecorder.cpp
--- a/Engine/Plugins/VirtualProduction/Takes/Source/TakeRecorder/Private/SerializedRecorder.cpp
+++ b/Engine/Plugins/VirtualProduction/Takes/Source/TakeRecorder/Private/SerializedRecorder.cpp
@@ -1,7 +1,6 @@
 // Copyright Epic Games, Inc. All Rights Reserved.
 
 #include "SerializedRecorder.h"
-#include "Editor.h"
 #include "TakeRecorderSources.h"
 #include "Serializers/MovieSceneSectionSerialization.h"
 #include "Serializers/MovieSceneActorSerialization.h"
@@ -25,14 +24,13 @@
 #include "GameFramework/MovementComponent.h"
 #include "GameFramework/Character.h"
 
-#include "IAssetTools.h"
 #include "AssetRegistry/AssetData.h"
-#include "AssetToolsModule.h"
-#include "Factories/Factory.h"
 #include "UObject/UObjectIterator.h"
 #include "LevelSequenceActor.h"
-#include "LevelEditorViewport.h"
+
+#if WITH_EDITOR
 #include "SequencerUtilities.h"
+#endif // WITH_EDITOR
 
 #define LOCTEXT_NAMESPACE "SerializedRecorder"
 
@@ -276,7 +274,9 @@ bool FSerializedRecorder::LoadSubSequenceFile(UMovieSceneSequence* InMovieSceneS
 				}
 
 				// We create a new sub track for every Source so that we can name the Subtrack after the Source instead of just the sections within it.
+			#if WITH_EDITOR
 				SubsceneTrack->SetDisplayName(FText::FromString(SubSequenceName));
+			#endif // WITH_EDITOR
 
 				// When we create the Subscene Track we'll make sure a folder is created for it to sort into and add the new Subscene Track as a child of it.
 				/*
@@ -418,15 +418,23 @@ AActor* FSerializedRecorder::SetActorPossesableOrSpawnable(UMovieSceneSequence*
 		}
 		//end of MakeSpawnableTemplateFromInstance hack
 
+	#if WITH_EDITOR
 		Actor->SetActorLabel(ActorHeader.Label); //spawnable has samen Name used for matching.
-		PossessableGuid = MovieScene->AddPossessable(ActorHeader.Label, SpawnableClass);
+	#endif // WITH_EDITOR
+
 
+	#if WITH_EDITOR
+		PossessableGuid = MovieScene->AddPossessable(ActorHeader.Label, SpawnableClass);
 		UE::Sequencer::FCreateBindingParams CreateBindingParams;
 		CreateBindingParams.bAllowCustomBinding = true;
 		CreateBindingParams.bSpawnable = true;
 		CreateBindingParams.BindingNameOverride = ActorHeader.Label;
 		CreateBindingParams.ReplacementGuid = PossessableGuid;
 		FSequencerUtilities::CreateOrReplaceBinding(nullptr, InMovieSceneSequence, Actor, CreateBindingParams);
+	#else // WITH_EDITOR
+		// Review the Guid here and how the binding is created at runtime.
+		PossessableGuid = MovieSceneHelpers::TryCreateCustomSpawnableBinding(InMovieSceneSequence, Actor);
+	#endif // WITH_EDITOR
 
 		if (PossessableGuid.IsValid())
 		{
@@ -478,6 +486,8 @@ AActor* FSerializedRecorder::SetActorPossesableOrSpawnable(UMovieSceneSequence*
 	if (ActorHeader.FolderName.IsValid())
 	{
 		UMovieSceneFolder* FolderToUse = nullptr;
+
+	#if WITH_EDITOR
 		for (UMovieSceneFolder* Folder : MovieScene->GetRootFolders())
 		{
 			if (Folder->GetFolderName() == ActorHeader.FolderName)
@@ -486,11 +496,15 @@ AActor* FSerializedRecorder::SetActorPossesableOrSpawnable(UMovieSceneSequence*
 				break;
 			}
 		}
+	#endif // WITH_EDITOR
+
 		if (FolderToUse == nullptr)
 		{
 			FolderToUse = NewObject<UMovieSceneFolder>(MovieScene, NAME_None, RF_Transactional);
 			FolderToUse->SetFolderName(ActorHeader.FolderName);
+		#if WITH_EDITOR
 			MovieScene->AddRootFolder(FolderToUse);
+		#endif // WITH_EDITOR
 		}
 
 		FolderToUse->AddChildObjectBinding(PossessableGuid);
diff --git a/Engine/Plugins/VirtualProduction/Takes/Source/TakeRecorder/Private/TakeRecorderDirectoryHelpers.cpp b/Engine/Plugins/VirtualProduction/Takes/Source/TakeRecorder/Private/TakeRecorderDirectoryHelpers.cpp
--- a/Engine/Plugins/VirtualProduction/Takes/Source/TakeRecorder/Private/TakeRecorderDirectoryHelpers.cpp
+++ b/Engine/Plugins/VirtualProduction/Takes/Source/TakeRecorder/Private/TakeRecorderDirectoryHelpers.cpp
@@ -2,7 +2,12 @@
 
 #include "TakeRecorderDirectoryHelpers.h"
 
+#if WITH_EDITOR
 #include "EditorDirectories.h"
+#else // WITH_EDITOR
+#include "Misc/Paths.h"
+#endif // WITH_EDITOR
+
 #include "Misc/PackageName.h"
 #include "HAL/PlatformFileManager.h"
 #include "HAL/PlatformFile.h"
@@ -42,6 +47,7 @@ TTuple<FString, FString> SplitPackagePath(const FString& InPath)
 
 FString GetDefaultProjectPath()
 {
+#if WITH_EDITOR
 	// Determine the starting path. Try to use the most recently used directory
 	FString AssetPath;
 	const FString DefaultFilesystemDirectory = FEditorDirectories::Get().GetLastDirectory(ELastDirectory::NEW_ASSET);
@@ -54,6 +60,11 @@ FString GetDefaultProjectPath()
 		return AssetPath.Left(AssetPath.Len()-1);
 	}
 	return AssetPath;
+#else // WITH_EDITOR
+	// When the editor is not present, return the Project Content Directory. This is the behaviour of default path of
+	// FEditorDirectories if there is no last directory, and so will mimick its behaviour.
+	return FPaths::ProjectContentDir();
+#endif // WITH_EDITOR
 }
 
 FString ResolvePathToProject(const FString& InPath)
diff --git a/Engine/Plugins/VirtualProduction/Takes/Source/TakeRecorder/Private/TakeRecorderModule.cpp b/Engine/Plugins/VirtualProduction/Takes/Source/TakeRecorder/Private/TakeRecorderModule.cpp
--- a/Engine/Plugins/VirtualProduction/Takes/Source/TakeRecorder/Private/TakeRecorderModule.cpp
+++ b/Engine/Plugins/VirtualProduction/Takes/Source/TakeRecorder/Private/TakeRecorderModule.cpp
@@ -4,46 +4,24 @@
 
 #include "TakesCoreLog.h"
 #include "TakeRecorderSettings.h"
-#include "SequencerUtilities.h"
-#include "Customization/NamingTokenDirectoryPathCustomization.h"
-#include "Customization/TakeRecorderProjectSettingsCustomization.h"
 #include "Modules/ModuleManager.h"
-#include "LevelEditor.h"
-#include "TakeRecorderCommands.h"
-#include "TakeRecorderStyle.h"
 #include "Recorder/TakeRecorder.h"
 
-#include "WorkspaceMenuStructureModule.h"
-#include "WorkspaceMenuStructure.h"
 #include "Widgets/Input/SButton.h"
 #include "Widgets/Docking/SDockTab.h"
 #include "Framework/Application/SlateApplication.h"
 #include "Framework/Docking/LayoutExtender.h"
 #include "Framework/MultiBox/MultiBoxExtender.h"
 #include "Framework/MultiBox/MultiBoxBuilder.h"
-#include "ToolMenus.h"
-#include "ContentBrowserMenuContexts.h"
 #include "Features/IModularFeatures.h"
 #include "ITakeRecorderDropHandler.h"
 #include "ISettingsModule.h"
 #include "TakeRecorderSettings.h"
 #include "TakePresetSettings.h"
 
-#include "Widgets/STakeRecorderTabContent.h"
-#include "Widgets/STakeRecorderCockpit.h"
-#include "Widgets/STakeRecorderPanel.h"
-
-#include "ISequencer.h"
-#include "ISequencerModule.h"
 #include "LevelSequence.h"
-#include "LevelSequenceEditorModule.h"
-#include "SequencerSettings.h"
 #include "TakeMetaData.h"
 #include "MovieSceneTakeSettings.h"
-#include "FileHelpers.h"
-
-#include "IContentBrowserSingleton.h"
-#include "ContentBrowserModule.h"
 
 #include "SerializedRecorder.h"
 #include "Misc/Attribute.h"
@@ -52,7 +30,6 @@
 #include "Misc/QualifiedFrameTime.h"
 #include "MovieSceneTakeSection.h"
 #include "MovieSceneTakeTrack.h"
-#include "ISequencer.h"
 #include "LevelSequence.h"
 #include "MovieSceneSequence.h"
 #include "Sections/MovieSceneSubSection.h"
@@ -63,6 +40,10 @@
 #include "Engine/Font.h"
 #include "CanvasTypes.h"
 
+#if WITH_EDITOR
+#include "Editor.h"
+#endif // WITH_EDITOR
+
 #define LOCTEXT_NAMESPACE "TakeRecorderModule"
 
 FName ITakeRecorderModule::TakeRecorderTabName = "TakeRecorder";
@@ -87,278 +68,12 @@ TArray<ITakeRecorderDropHandler*> ITakeRecorderDropHandler::GetDropHandlers()
 	return IModularFeatures::Get().GetModularFeatureImplementations<ITakeRecorderDropHandler>(ModularFeatureName);
 }
 
-namespace UE::TakeRecorder::Private
-{
-    template<typename T>
-    concept CInvocableWithMovieSceneTakeSection = std::is_invocable_v<T, UMovieSceneTakeSection*>;
-
-	/** Render the timecode data with slate and rate to the canvas. Flag invalid TC values with a red color to indicate a problem to the user. */
-	int32 RenderTimecode(FCanvas* Canvas, int32 X, int32 Y, const FTimecode& Timecode, const float& Rate, const FString& SequenceName)
-	{
-		UFont* Font = FPlatformProperties::SupportsWindowedMode() ? GEngine->GetSmallFont() : GEngine->GetMediumFont();
-		const int32 RowHeight = FMath::TruncToInt(Font->GetMaxCharHeight());
-
-		const bool bForceSignDisplay = false;
-		const bool bAlwaysDisplaySubframe = true;
-
-		const FString TimecodeStr = Timecode.ToString(bForceSignDisplay, bAlwaysDisplaySubframe);
-
-		const FString TakeSection = TEXT("Take Section -- ");
-		const FString FPS = TEXT("(00.00)");
-		float CharWidth, CharHeight;
-		Font->GetCharSize(TEXT(' '), CharWidth, CharHeight);
-
-		int32 TakeSectionWidth = Font->GetStringSize(*TakeSection);
-		int32 NewX = X - Font->GetStringSize(*SequenceName) - Font->GetStringSize(*FPS) - TakeSectionWidth - (int32)CharWidth;
-		int32 SectionX = NewX + TakeSectionWidth;
-		FColor Color = Rate > 6000 ? FColor::Red : FColor::Green;
-		float DisplayRate = Rate > 6000 ? 0 : Rate;
-		Canvas->DrawShadowedString(NewX, Y, *TakeSection, Font, FColor::Cyan);
-		Canvas->DrawShadowedString(SectionX, Y, *FString::Printf(TEXT("%s TC: %s (%.2f)"), *SequenceName, *TimecodeStr, DisplayRate), Font, Color);
-		Y += RowHeight;
-		return Y;
-	};
-
-	template <typename CInvocableWithMovieSceneTakeSection>
-	void IterateOverMovieSceneForSections(UMovieScene* MovieScene, CInvocableWithMovieSceneTakeSection&& SectionFunction)
-	{
-		if (!MovieScene)
-		{
-			return;
-		}
-
-		// Lambda to iterate over tracks.  This will either get called via a FMovieSceneBinding or directly froma UMovieScene
-		auto ForEachTrack = [&SectionFunction](const TArray<UMovieSceneTrack*>& Tracks) -> bool
-		{
-			bool bDidRenderSection = false;
-			for (UMovieSceneTrack* Track : Tracks)
-			{
-				for (UMovieSceneSection* Section : Track->GetAllSections())
-				{
-					if (UMovieSceneTakeSection* TakeSection = Cast<UMovieSceneTakeSection>(Section))
-					{
-						SectionFunction(TakeSection);
-						bDidRenderSection = true;
-					}
-					else if (UMovieSceneSubSection* SubSection = Cast<UMovieSceneSubSection>(Section))
-					{
-						if (SubSection->GetSequence())
-						{
-							IterateOverMovieSceneForSections(SubSection->GetSequence()->GetMovieScene(), SectionFunction);
-						}
-					}
-				}
-			}
-			return bDidRenderSection;
-		};
-
-		const TArray<FMovieSceneBinding>& Bindings = ((const UMovieScene*)MovieScene)->GetBindings();
-		bool bDidRenderForMovieScene = false;
-		for (const FMovieSceneBinding& Binding : Bindings)
-		{
-			bDidRenderForMovieScene = ForEachTrack(Binding.GetTracks());
-		}
-
-		// If we didn't render anything with FMovieSceneBinding try the MovieScene object.
-		if (!bDidRenderForMovieScene)
-		{
-			ForEachTrack(MovieScene->GetTracks());
-		}
-	}
-
-	int32 RenderTakeSectionsInSequencer(FCanvas* Canvas, int32 X, int32 Y, TSharedPtr<ISequencer>& InSequencer)
-	{
-		const FFrameRate RootTickRate = InSequencer->GetRootTickResolution();
-		const FFrameTime CurrentTime = InSequencer->GetGlobalTime().ConvertTo(RootTickRate);
-		auto RenderOneTakeSubSection = [CurrentTime, Canvas, X, &Y](UMovieSceneTakeSection* TakeSection) mutable
-		{
-			TOptional<UMovieSceneTakeSection::FSectionData> TakeData = TakeSection->Evaluate(CurrentTime);
-			if (TakeData)
-			{
-				Y = RenderTimecode(Canvas, X, Y, TakeData->Timecode, TakeData->Rate, TakeData->Slate);
-			}
-		};
-
-		const TArray<FMovieSceneSequenceID>& SubSequenceHierarchy = InSequencer->GetSubSequenceHierarchy();
-		if (SubSequenceHierarchy.Num()>0)
-		{
-			UMovieSceneSequence* Sequence = FSequencerUtilities::GetMovieSceneSequence(InSequencer, SubSequenceHierarchy.Last());
-			check(Sequence);
-
-			UMovieScene* MovieScene = Sequence->GetMovieScene();
-			IterateOverMovieSceneForSections(MovieScene, RenderOneTakeSubSection);
-		}
-		return Y;
-	}
-
-	static FOpenSequencerWatcher SequencerWatcher;
-
-	/** For the given sequencer, iterate over all sections and find a take section then evaluate and show timecode, rate, and slate info in the HUD. */
-	int32 RenderTakeSectionTime(UWorld* World, FViewport* Viewport, FCanvas* Canvas, int32 X, int32 Y, const FVector* ViewLocation, const FRotator* ViewRotation)
-	{
-		for (const FOpenSequencerWatcher::FOpenSequencerData& OpenSequencer : SequencerWatcher.OpenSequencers)
-		{
-			if (TSharedPtr<ISequencer> Sequencer = OpenSequencer.WeakSequencer.Pin())
-			{
-				Y = RenderTakeSectionsInSequencer(Canvas, X, Y, Sequencer);
-			}
-		}
-		return Y;
-	}
-
-    void InitStatCommands()
-    {
-		auto StartupComplete = []()
-		{
-			check(GEngine);
-			if (GIsEditor)
-			{
-                const bool bIsRHS = true;
-				GEngine->AddEngineStat(TEXT("STAT_TakeTimecode"), TEXT("STATCAT_Sequencer"),
-									   LOCTEXT("TakeTimecodeDisplay", "Displays current sequencer time value in NDF timecode format."),
-									   UEngine::FEngineStatRender::CreateStatic(&RenderTakeSectionTime), nullptr, bIsRHS);
-			}
-		};
-
-		SequencerWatcher.DoStartup(StartupComplete);
-    }
-}
-
-namespace
-{
-	static TSharedRef<SDockTab> SpawnTakesBrowserTab(const FSpawnTabArgs& SpawnTabArgs)
-	{
-		TSharedRef<SDockTab> NewTab = SNew(SDockTab)
-		.TabRole(ETabRole::NomadTab);
-
-		FContentBrowserConfig ContentBrowserConfig;
-		{
-			ContentBrowserConfig.ThumbnailLabel =  EThumbnailLabel::ClassName ;
-			ContentBrowserConfig.ThumbnailScale = 0.1f;
-			ContentBrowserConfig.InitialAssetViewType = EAssetViewType::Column;
-			ContentBrowserConfig.bShowBottomToolbar = true;
-			ContentBrowserConfig.bCanShowClasses = true;
-			ContentBrowserConfig.bUseSourcesView = true;
-			ContentBrowserConfig.bExpandSourcesView = true;
-			ContentBrowserConfig.bUsePathPicker = true;
-			ContentBrowserConfig.bCanShowFilters = true;
-			ContentBrowserConfig.bCanShowAssetSearch = true;
-			ContentBrowserConfig.bCanShowFolders = true;
-			ContentBrowserConfig.bCanShowRealTimeThumbnails = true;
-			ContentBrowserConfig.bCanShowDevelopersFolder = true;
-			ContentBrowserConfig.bCanShowLockButton = true;
-			ContentBrowserConfig.bCanSetAsPrimaryBrowser = false;
-		}
-
-		IContentBrowserSingleton& ContentBrowser = FModuleManager::LoadModuleChecked<FContentBrowserModule>("ContentBrowser").Get();
-		TSharedRef<SWidget> NewBrowser = ContentBrowser.CreateContentBrowser( FTakeRecorderModule::TakesBrowserInstanceName, NewTab, nullptr );
-
-		NewTab->SetContent( NewBrowser );
-
-		FString TakesDir = FPaths::GetPath(FPaths::GetPath(GetDefault<UTakeRecorderProjectSettings>()->Settings.GetTakeAssetPath()));
-		TArray<FString> TakesFolders;
-		TakesFolders.Push(TakesDir);
-		ContentBrowser.SyncBrowserToFolders(TakesFolders, true, false, FTakeRecorderModule::TakesBrowserInstanceName);
-
-		return NewTab;
-	}
-
-	static TSharedRef<SDockTab> SpawnTakeRecorderTab(const FSpawnTabArgs& SpawnTabArgs)
-	{
-		TSharedRef<STakeRecorderTabContent> Content = SNew(STakeRecorderTabContent);
-		TSharedPtr<SDockTab> ContentTab = SNew(SDockTab)
-			.Label(Content, &STakeRecorderTabContent::GetTitle)
-			.TabRole(ETabRole::NomadTab)
-			[
-				Content
-			];
-		const TAttribute<const FSlateBrush*> TabIcon = TAttribute<const FSlateBrush*>::CreateLambda([Content]()
-			{
-				return Content->GetIcon();
-			});
-		ContentTab->SetTabIcon(TabIcon);
-		return ContentTab.ToSharedRef();
-	}
-
-	static void RegisterLevelEditorLayout(FLayoutExtender& Extender)
-	{
-		Extender.ExtendArea("TopLevelArea",
-			[](TSharedRef<FTabManager::FArea> InArea)
-			{
-				InArea->SplitAt(1, 
-					FTabManager::NewStack()
-					->SetSizeCoefficient( 0.3f )
-					->AddTab(ITakeRecorderModule::TakeRecorderTabName, ETabState::ClosedTab)
-				);
-			}
-		);
-	}
-
-	static void RegisterTabImpl()
-	{
-		FTabSpawnerEntry& TabSpawner = FGlobalTabmanager::Get()->RegisterNomadTabSpawner(ITakeRecorderModule::TakeRecorderTabName, FOnSpawnTab::CreateStatic(SpawnTakeRecorderTab));
-
-		TabSpawner
-			.SetGroup(WorkspaceMenu::GetMenuStructure().GetLevelEditorCinematicsCategory())
-			.SetDisplayName(ITakeRecorderModule::TakeRecorderTabLabel)
-			.SetTooltipText(LOCTEXT("TakeRecorderTab_Tooltip", "Open the main Take Recorder UI."))
-			.SetIcon(FSlateIcon(FAppStyle::GetAppStyleSetName(), "SequenceRecorder.TabIcon"));
-
-		FTabSpawnerEntry& TBTabSpawner = FGlobalTabmanager::Get()->RegisterNomadTabSpawner(ITakeRecorderModule::TakesBrowserTabName, FOnSpawnTab::CreateStatic(SpawnTakesBrowserTab));
-
-		TBTabSpawner
-			.SetGroup(WorkspaceMenu::GetMenuStructure().GetLevelEditorCinematicsCategory())
-			.SetDisplayName(ITakeRecorderModule::TakesBrowserTabLabel)
-			.SetTooltipText(LOCTEXT("TakeBrowserTab_Tooltip", "Open the Take Browser UI"))
-			.SetIcon(FSlateIcon(FAppStyle::GetAppStyleSetName(), "ContentBrowser.TabIcon"));
-	}
-	
-	static void ModulesChangedCallback(FName ModuleName, EModuleChangeReason ReasonForChange)
-	{
-		static const FName LevelEditorModuleName(TEXT("LevelEditor"));
-		if (ReasonForChange == EModuleChangeReason::ModuleLoaded && ModuleName == LevelEditorModuleName)
-		{
-			RegisterTabImpl();
-		}
-	}
-}
-
-FTakeRecorderModule::FTakeRecorderModule()
-	: SequencerSettings(nullptr)
-{
-}
-
 void FTakeRecorderModule::StartupModule()
 {
-	FTakeRecorderStyle::Get();
-	FTakeRecorderCommands::Register();
-
-	RegisterDetailCustomizations();
-	RegisterLevelEditorExtensions();
 	RegisterSettings();
 	RegisterSerializedRecorder();
-	TimecodeManagement = MakeUnique<UE::TakeRecorder::FHitchlessProtectionRootLogic>();
-#if WITH_EDITOR
-	if (GIsEditor)
-	{
-		SequencerHitchVisualizer = MakeUnique<UE::TakeRecorder::FSequencerHitchVisualizer>();
-	}
-#endif
 
 #if WITH_EDITOR
-	if (GIsEditor)
-	{
-		if (UToolMenus::TryGet())
-		{
-			RegisterMenus();
-		}
-		else
-		{
-			FCoreDelegates::GetOnPostEngineInit().AddRaw(this, &FTakeRecorderModule::RegisterMenus);
-		}
-		UE::TakeRecorder::Private::InitStatCommands();
-	}
-
 	if (GEditor)
 	{
 		GEditor->OnEditorClose().AddRaw(this, &FTakeRecorderModule::OnEditorClose);
@@ -378,16 +93,13 @@ void FTakeRecorderModule::ShutdownModule()
 	FCoreDelegates::GetOnPostEngineInit().RemoveAll(this);
 #endif
 
-	FTakeRecorderCommands::Unregister();
-
-	UnregisterDetailCustomizations();
-	UnregisterLevelEditorExtensions();
 	UnregisterSettings();
 	UnregisterSerializedRecorder();
-	TimecodeManagement.Reset();
-#if WITH_EDITOR
-	SequencerHitchVisualizer.Reset();
-#endif
+}
+
+void FTakeRecorderModule::PopulateSourcesMenu(TSharedRef<FExtender> InExtender, UTakeRecorderSources* InSources) const
+{
+	SourcesMenuExtenderEvent.Broadcast(InExtender, InSources);
 }
 
 UTakePreset* FTakeRecorderModule::GetPendingTake() const
@@ -437,91 +149,6 @@ void FTakeRecorderModule::RegisterSettingsObject(UObject* InSettingsObject)
 	GetMutableDefault<UTakeRecorderProjectSettings>()->AdditionalSettings.Add(InSettingsObject);
 }
 
-void FTakeRecorderModule::AddReferencedObjects(FReferenceCollector& Collector)
-{
-	if (SequencerSettings)
-	{
-		Collector.AddReferencedObject(SequencerSettings);
-	}
-}
-
-void FTakeRecorderModule::RegisterDetailCustomizations()
-{
-#if WITH_EDITOR
-	if (GIsEditor)
-	{
-		FPropertyEditorModule& PropertyEditorModule = FModuleManager::Get().LoadModuleChecked<FPropertyEditorModule>("PropertyEditor");
-		PropertyEditorModule.RegisterCustomClassLayout(
-			UTakeRecorderProjectSettings::StaticClass()->GetFName(),
-			FOnGetDetailCustomizationInstance::CreateLambda(&MakeShared<FTakeRecorderProjectSettingsCustomization>)
-			);
-
-		PropertyEditorModule.RegisterCustomPropertyTypeLayout(FName(TEXT("DirectoryPath")),
-			FOnGetPropertyTypeCustomizationInstance::CreateLambda(&FNamingTokenDirectoryPathStructCustomization::MakeInstance),
-			MakeShared<UE::TakeRecorder::NamingTokens::FDirectoryPathPropertyTypeIdentifier>());
-	}
-#endif
-}
-
-void FTakeRecorderModule::UnregisterDetailCustomizations()
-{
-	if (UObjectInitialized() && !IsEngineExitRequested())
-	{
-		if (FPropertyEditorModule* PropertyEditorModule = FModuleManager::Get().GetModulePtr<FPropertyEditorModule>("PropertyEditor"))
-		{
-			PropertyEditorModule->UnregisterCustomClassLayout(UTakeRecorderProjectSettings::StaticClass()->GetFName());
-			PropertyEditorModule->UnregisterCustomPropertyTypeLayout(FName(TEXT("DirectoryPath")),
-				MakeShared<UE::TakeRecorder::NamingTokens::FDirectoryPathPropertyTypeIdentifier>());
-		}
-	}
-}
-
-void FTakeRecorderModule::RegisterLevelEditorExtensions()
-{
-#if WITH_EDITOR
-
-	if (GIsEditor)
-	{
-		FLevelEditorModule& LevelEditorModule = FModuleManager::LoadModuleChecked<FLevelEditorModule>("LevelEditor");
-
-		LevelEditorLayoutExtensionHandle = LevelEditorModule.OnRegisterLayoutExtensions().AddStatic(RegisterLevelEditorLayout);
-
-		if (LevelEditorModule.GetLevelEditorTabManager())
-		{
-			RegisterTabImpl();
-		}
-		else
-		{
-			LevelEditorTabManagerChangedHandle = LevelEditorModule.OnTabManagerChanged().AddStatic(RegisterTabImpl);
-		}
-
-		if (!FModuleManager::Get().IsModuleLoaded("LevelEditor"))
-		{
-			ModulesChangedHandle = FModuleManager::Get().OnModulesChanged().AddStatic(ModulesChangedCallback);
-		}
-	}
-
-#endif
-}
-
-void FTakeRecorderModule::UnregisterLevelEditorExtensions() const
-{
-#if WITH_EDITOR
-	if (FSlateApplication::IsInitialized())
-	{
-		FGlobalTabmanager::Get()->UnregisterNomadTabSpawner(TakeRecorderTabName);
-		FGlobalTabmanager::Get()->UnregisterNomadTabSpawner(TakesBrowserTabName);
-	}
-#endif
-
-	if(FLevelEditorModule* LevelEditorModulePtr = FModuleManager::GetModulePtr<FLevelEditorModule>("LevelEditor"))
-	{
-		LevelEditorModulePtr->OnTabManagerChanged().Remove(LevelEditorTabManagerChangedHandle);
-	}
-
-	FModuleManager::Get().OnModulesChanged().Remove(ModulesChangedHandle);
-}
-
 void FTakeRecorderModule::RegisterSettings()
 {
 	RegisterSettingsObject(GetMutableDefault<UMovieSceneTakeSettings>());
@@ -541,19 +168,12 @@ void FTakeRecorderModule::RegisterSettings()
 		GetMutableDefault<UTakeRecorderUserSettings>()
 	);
 
-	SequencerSettings = USequencerSettingsContainer::GetOrCreate<USequencerSettings>(TEXT("TakeRecorderSequenceEditor"));
-
 	GetMutableDefault<UTakeRecorderUserSettings>()->LoadConfig();
 	const bool bSaveRecordedAssetsOverride = CVarTakeRecorderSaveRecordedAssetsOverride.GetValueOnGameThread() != 0;
 	if (bSaveRecordedAssetsOverride)
 	{
 		GetMutableDefault<UTakeRecorderUserSettings>()->Settings.bSaveRecordedAssets = bSaveRecordedAssetsOverride;
 	}
-
-	SettingsModule.RegisterSettings("Editor", "ContentEditors", "TakeRecorderSequenceEditor",
-		LOCTEXT("TakeRecorderSequenceEditorSettingsName", "Take Recorder Sequence Editor"),
-		LOCTEXT("TakeRecorderSequenceEditorSettingsDescription", "Configure the look and feel of the Take Recorder Sequence Editor."),
-		SequencerSettings);
 }
 
 void FTakeRecorderModule::UnregisterSettings()
@@ -563,15 +183,9 @@ void FTakeRecorderModule::UnregisterSettings()
 	{
 		SettingsModule->UnregisterSettings("Project", "Plugins", "Take Recorder");
 		SettingsModule->UnregisterSettings("Editor",  "ContentEditors", "Take Recorder");
-		SettingsModule->UnregisterSettings("Editor", "ContentEditors", "TakeRecorderSequenceEditor");
 	}
 }
 
-void FTakeRecorderModule::PopulateSourcesMenu(TSharedRef<FExtender> InExtender, UTakeRecorderSources* InSources) const
-{
-	SourcesMenuExtenderEvent.Broadcast(InExtender, InSources);
-}
-
 void FTakeRecorderModule::RegisterSerializedRecorder()
 {
 	SerializedRecorder = MakeShared<FSerializedRecorder>();
@@ -583,93 +197,6 @@ void FTakeRecorderModule::UnregisterSerializedRecorder() const
 	IModularFeatures::Get().UnregisterModularFeature(FSerializedRecorder::ModularFeatureName, SerializedRecorder.Get());
 }
 
-void FTakeRecorderModule::RegisterMenus()
-{
-#if WITH_EDITOR
-	if (!UToolMenus::IsToolMenuUIEnabled())
-	{
-		return;
-	}
-
-	FToolMenuOwnerScoped MenuOwner("TakeRecorder");
-	UToolMenus* ToolMenus = UToolMenus::Get();
-	UToolMenu* Menu = ToolMenus->ExtendMenu("ContentBrowser.AssetContextMenu.LevelSequence");
-	if (!Menu)
-	{
-		return;
-	}
-
-	FToolMenuSection& Section = Menu->FindOrAddSection("GetAssetActions");
-	Section.AddDynamicEntry("TakeRecorderActions", FNewToolMenuSectionDelegate::CreateLambda([](FToolMenuSection& InSection)
-	{
-		UContentBrowserAssetContextMenuContext* Context = InSection.FindContext<UContentBrowserAssetContextMenuContext>();
-		if (!Context)
-		{
-			return;
-		}
-
-		if (Context->SelectedAssets.Num() == 1 && Context->SelectedAssets[0].IsInstanceOf(ULevelSequence::StaticClass()))
-		{
-			const FAssetData LevelSequenceAsset = Context->SelectedAssets[0];
-
-			InSection.AddMenuEntry(
-				"OpenInTakeRecorder_Label",
-				LOCTEXT("OpenInTakeRecorder_Label", "Open in Take Recorder"),
-				LOCTEXT("OpenInTakeRecorder_Tooltip", "Opens this level sequence asset in Take Recorder by copying its contents into the pending take"),
-				FSlateIcon(FTakeRecorderStyle::StyleName, "TakeRecorder.TabIcon"),
-				FExecuteAction::CreateLambda(
-					[LevelSequenceAsset]
-					{
-						if (ULevelSequence* LevelSequence = Cast<ULevelSequence>(LevelSequenceAsset.GetAsset()))
-						{
-							FLevelEditorModule& LevelEditorModule = FModuleManager::LoadModuleChecked<FLevelEditorModule>("LevelEditor");
-							TSharedPtr<SDockTab> DockTab = LevelEditorModule.GetLevelEditorTabManager()->TryInvokeTab(ITakeRecorderModule::TakeRecorderTabName);
-							if (DockTab.IsValid())
-							{
-								TSharedRef<STakeRecorderTabContent> TabContent = StaticCastSharedRef<STakeRecorderTabContent>(DockTab->GetContent());
-
-								// If this sequence has already been recorded, set it up for viewing, otherwise start recording from it.
-								UTakeMetaData* TakeMetaData = LevelSequence->FindMetaData<UTakeMetaData>();
-								if (!TakeMetaData || !TakeMetaData->Recorded())
-								{
-									TabContent->SetupForRecording(LevelSequence);
-								}
-								else
-								{
-									TabContent->SetupForViewing(LevelSequence);
-								}
-							}
-						}
-					}
-				)
-			);
-			InSection.AddMenuEntry(
-				"RecordIntoTakeRecorder_Label",
-				LOCTEXT("RecordWithTakeRecorder_Label", "Record with Take Recorder"),
-				LOCTEXT("RecordWithTakeRecorder_Tooltip", "Opens this level sequence asset for recording into with Take Recorder"),
-				FSlateIcon(FTakeRecorderStyle::StyleName, "TakeRecorder.TabIcon"),
-				FExecuteAction::CreateLambda(
-					[LevelSequenceAsset]
-					{
-						if (ULevelSequence* LevelSequence = Cast<ULevelSequence>(LevelSequenceAsset.GetAsset()))
-						{
-							FLevelEditorModule& LevelEditorModule = FModuleManager::LoadModuleChecked<FLevelEditorModule>("LevelEditor");
-							TSharedPtr<SDockTab> DockTab = LevelEditorModule.GetLevelEditorTabManager()->TryInvokeTab(ITakeRecorderModule::TakeRecorderTabName);
-							if (DockTab.IsValid())
-							{
-								TSharedRef<STakeRecorderTabContent> TabContent = StaticCastSharedRef<STakeRecorderTabContent>(DockTab->GetContent());
-
-								TabContent->SetupForRecordingInto(LevelSequence);
-							}
-						}
-					}
-				)
-			);
-		}
-	}));
-#endif // WITH_EDITOR
-}
-
 void FTakeRecorderModule::OnEditorClose()
 {
 	if (UTakeRecorder* ActiveRecorder = UTakeRecorder::GetActiveRecorder())
diff --git a/Engine/Plugins/VirtualProduction/Takes/Source/TakeRecorder/Private/TakeRecorderModule.h b/Engine/Plugins/VirtualProduction/Takes/Source/TakeRecorder/Private/TakeRecorderModule.h
--- a/Engine/Plugins/VirtualProduction/Takes/Source/TakeRecorder/Private/TakeRecorderModule.h
+++ b/Engine/Plugins/VirtualProduction/Takes/Source/TakeRecorder/Private/TakeRecorderModule.h
@@ -5,10 +5,7 @@
 #include "Delegates/IDelegateInstance.h"
 #include "ITakeRecorderModule.h"
 #include "Templates/SharedPointer.h"
-#include "Timecode/HitchlessProtectionRootLogic.h"
-#include "Timecode/Visualization/Sequencer/SequencerHitchVisualizer.h"
 #include "UObject/GCObject.h"
-#include "Timecode/HitchlessProtectionRootLogic.h"
 #include "UObject/WeakObjectPtrTemplates.h"
 
 class FExtender;
@@ -18,19 +15,16 @@ class FSerializedRecorder;
 class UTakeRecorderSources;
 class USequencerSettings;
 
-class FTakeRecorderModule : public ITakeRecorderModule, public FGCObject
+class FTakeRecorderModule : public ITakeRecorderModule
 {
 public:
-	FTakeRecorderModule();
-
-	void PopulateSourcesMenu(TSharedRef<FExtender> InExtender, UTakeRecorderSources* InSources) const;
-
 	//~ Begin IModuleInterface Interface
 	virtual void StartupModule() override;
 	virtual void ShutdownModule() override;
 	//~ End IModuleInterface Interface
 	
 	//~ Begin ITakeRecorderModule Interface
+	virtual void PopulateSourcesMenu(TSharedRef<FExtender> InExtender, UTakeRecorderSources* InSources) const override;
 	virtual UTakePreset* GetPendingTake() const override;
 	virtual FOnGenerateWidgetExtensions& GetToolbarExtensionGenerators() override { return ToolbarExtensionGenerators; }
 	virtual FOnGenerateWidgetExtensions& GetRecordButtonExtensionGenerators() override { return ButtonExtensionGenerators; }
@@ -47,31 +41,17 @@ public:
 	virtual void UnregisterSourcesExtension() override;
 	virtual void UnregisterSourcesMenuExtension(FDelegateHandle Handle) override;
 	virtual void RegisterSettingsObject(UObject* InSettingsObject) override;
+	virtual FSourceExtensionData& GetSourcesExtensionData() override { return SourceExtensionData; }
 	//~ Begin ITakeRecorderModule Interface
 	
-	FSourceExtensionData& GetSourcesExtensionData() { return SourceExtensionData; }
 	
 private:
-	
-	//~ Begin FGCObject Interface
-	virtual void AddReferencedObjects(FReferenceCollector& Collector) override;
-	virtual FString GetReferencerName() const override { return "FTakeRecorderModule"; }
-	//~ End FGCObject Interface
-
-	void RegisterDetailCustomizations();
-	void UnregisterDetailCustomizations();
-
-	void RegisterLevelEditorExtensions();
-	void UnregisterLevelEditorExtensions() const;
-
 	void RegisterSettings();
 	void UnregisterSettings();
 
 	void RegisterSerializedRecorder();
 	void UnregisterSerializedRecorder() const;
 
-	void RegisterMenus();
-
 	void OnEditorClose();
 
 private:
@@ -94,24 +74,9 @@ private:
 
 	TSharedPtr<FSerializedRecorder> SerializedRecorder;
 
-	/**
-	 * Sets up all sub-systems related to managing timecode in Take Recorder
-	 * 
-	 * For future maintainers: this could totally live in a separate module / plugin / whatever.
-	 * There's no technical reason to have this live in TakeRecorder. The reason it's here is that we probably always want hitching to be handled
-	 * gracefully.
-	 */
-	TUniquePtr<UE::TakeRecorder::FHitchlessProtectionRootLogic> TimecodeManagement;
-#if WITH_EDITOR
-	TUniquePtr<UE::TakeRecorder::FSequencerHitchVisualizer> SequencerHitchVisualizer;
-#endif
-
 	/** Extra class default objects that should appear in the Take Recorder panel. */
 	TArray<TWeakObjectPtr<>> ExternalObjects;
 
-	/** The settings object registers with the Settings module. */
-	TObjectPtr<USequencerSettings> SequencerSettings;
-
 	/**
 	 * Registered extensions for the source module. We only use a single struct rather than container
 	 * as only the source module should be utilizing this.
diff --git a/Engine/Plugins/VirtualProduction/Takes/Source/TakeRecorder/Public/ITakeRecorderDropHandler.h b/Engine/Plugins/VirtualProduction/Takes/Source/TakeRecorder/Public/ITakeRecorderDropHandler.h
--- a/Engine/Plugins/VirtualProduction/Takes/Source/TakeRecorder/Public/ITakeRecorderDropHandler.h
+++ b/Engine/Plugins/VirtualProduction/Takes/Source/TakeRecorder/Public/ITakeRecorderDropHandler.h
@@ -19,7 +19,7 @@ struct ITakeRecorderDropHandler : public IModularFeature
 	/**
 	 * Get all the currently registered drop handlers
 	 */
-	static TArray<ITakeRecorderDropHandler*> GetDropHandlers();
+	TAKERECORDER_API static TArray<ITakeRecorderDropHandler*> GetDropHandlers();
 
 public:
 
diff --git a/Engine/Plugins/VirtualProduction/Takes/Source/TakeRecorder/Public/ITakeRecorderModule.h b/Engine/Plugins/VirtualProduction/Takes/Source/TakeRecorder/Public/ITakeRecorderModule.h
--- a/Engine/Plugins/VirtualProduction/Takes/Source/TakeRecorder/Public/ITakeRecorderModule.h
+++ b/Engine/Plugins/VirtualProduction/Takes/Source/TakeRecorder/Public/ITakeRecorderModule.h
@@ -89,19 +89,19 @@ public:
 	};
 	
 	/** The name under which the take recorder tab is registered and invoked */
-	static FName TakeRecorderTabName;
+	TAKERECORDER_API static FName TakeRecorderTabName;
 
 	/** The default label for the take recorder tab */
-	static FText TakeRecorderTabLabel;
+	TAKERECORDER_API static FText TakeRecorderTabLabel;
 
 	/** The tab name for the takes browser tab */
-	static FName TakesBrowserTabName;
+	TAKERECORDER_API static FName TakesBrowserTabName;
 
 	/** The default label for the takes browser */
-	static FText TakesBrowserTabLabel;
+	TAKERECORDER_API static FText TakesBrowserTabLabel;
 
 	/** The Takes Browser Content Browser Instance Name */
-	static FName TakesBrowserInstanceName;
+	TAKERECORDER_API static FName TakesBrowserInstanceName;
 
 	static inline ITakeRecorderModule& Get() 
 	{ 
@@ -113,6 +113,8 @@ public:
 		return FModuleManager::Get().IsModuleLoaded( "TakeRecorder" ); 
 	}
 
+	virtual void PopulateSourcesMenu(TSharedRef<FExtender> InExtender, UTakeRecorderSources* InSources) const = 0;
+
 	/** Get the current pending take. */
 	virtual UTakePreset* GetPendingTake() const = 0;
 
@@ -194,5 +196,7 @@ public:
 	 * Get the take external objects registered to the take recorder.
 	 */
 	virtual TArray<TWeakObjectPtr<>>& GetExternalObjects() = 0;
+
+	virtual FSourceExtensionData& GetSourcesExtensionData() = 0;
 };
 
diff --git a/Engine/Plugins/VirtualProduction/Takes/Source/TakeRecorder/Public/Recorder/TakeRecorder.h b/Engine/Plugins/VirtualProduction/Takes/Source/TakeRecorder/Public/Recorder/TakeRecorder.h
--- a/Engine/Plugins/VirtualProduction/Takes/Source/TakeRecorder/Public/Recorder/TakeRecorder.h
+++ b/Engine/Plugins/VirtualProduction/Takes/Source/TakeRecorder/Public/Recorder/TakeRecorder.h
@@ -4,14 +4,21 @@
 
 #include "UObject/Object.h"
 #include "UObject/StrongObjectPtr.h"
-#include "ISequencer.h"
 #include "Recorder/TakeRecorderParameters.h"
 #include "Serializers/MovieSceneManifestSerialization.h"
+
+#if WITH_EDITOR
+#include "ISequencer.h"
+#endif // WITH_EDITOR
+
 #include "TakeRecorder.generated.h"
 
 #define UE_API TAKERECORDER_API
 
+#if WITH_EDITOR
 class ISequencer;
+#endif // WITH_EDITOR
+
 class UTakePreset;
 class UTakeMetaData;
 class UTakeRecorder;
@@ -207,10 +214,12 @@ private:
 	 */
 	UE_API void Tick(float DeltaTime);
 
+#if WITH_EDITOR
 	/**
 	 * Called if we're currently recording a PIE world that has been shut down or if we start PIE in a non-PIE world. Bound in Initialize, and unbound in Stop.
 	 */
 	UE_API void HandlePIE(bool bIsSimulating);
+#endif // WITH_EDITOR
 
 	/**
 	 * Create a new destination asset to record into based on the parameters
@@ -305,8 +314,10 @@ private:
 	/** Triggered when this recorder ticks */
 	FOnTickRecording OnTickRecordingEvent;
 
+#if WITH_EDITOR
 	/** Sequencer ptr that controls playback of the destination asset during the recording */
 	TWeakPtr<ISequencer> WeakSequencer;
+#endif // WITH_EDITOR
 
 	/** Due a few ticks after the pre so we are set up with asset creation */
 	int32 NumberOfTicksAfterPre;
@@ -327,8 +338,11 @@ private:
 
 	FManifestSerializer ManifestSerializer;
 
+#if WITH_EDITOR
 	EAllowEditsMode CachedAllowEditsMode;
 	EAutoChangeMode CachedAutoChangeMode;
+#endif // WITH_EDITOR
+
 	EUpdateClockSource CachedClockSource;
 	
 	TRange<FFrameNumber> CachedPlaybackRange;
diff --git a/Engine/Plugins/VirtualProduction/Takes/Source/TakeRecorder/Public/Recorder/TakeRecorderPanel.h b/Engine/Plugins/VirtualProduction/Takes/Source/TakeRecorder/Public/Recorder/TakeRecorderPanel.h
--- a/Engine/Plugins/VirtualProduction/Takes/Source/TakeRecorder/Public/Recorder/TakeRecorderPanel.h
+++ b/Engine/Plugins/VirtualProduction/Takes/Source/TakeRecorder/Public/Recorder/TakeRecorderPanel.h
@@ -12,7 +12,7 @@ class ULevelSequence;
 class UTakeMetaData;
 class UTakeRecorderSources;
 
-class STakeRecorderTabContent;
+class ITakeRecorderTabContent;
 
 UENUM()
 enum class ETakeRecorderPanelMode : uint8
@@ -157,7 +157,7 @@ public:
 	/**
 	 * Initialize this object with a weak pointer to the tab content from which all the UI state can be retrieved
 	 */
-	UE_API void InitializePanel(TWeakPtr<STakeRecorderTabContent> InTabContent);
+	UE_API void InitializePanel(TWeakPtr<ITakeRecorderTabContent> InTabContent);
 
 	/**
 	 * Check whether this panel is still open or not
@@ -173,7 +173,7 @@ private:
 
 	UE_API bool ValidateTabContent() const;
 
-	TWeakPtr<STakeRecorderTabContent> WeakTabContent;
+	TWeakPtr<ITakeRecorderTabContent> WeakTabContent;
 };
 
 #undef UE_API
diff --git a/Engine/Plugins/VirtualProduction/Takes/Source/TakeRecorder/Public/TakeRecorderSettings.h b/Engine/Plugins/VirtualProduction/Takes/Source/TakeRecorder/Public/TakeRecorderSettings.h
--- a/Engine/Plugins/VirtualProduction/Takes/Source/TakeRecorder/Public/TakeRecorderSettings.h
+++ b/Engine/Plugins/VirtualProduction/Takes/Source/TakeRecorder/Public/TakeRecorderSettings.h
@@ -25,10 +25,10 @@ public:
 #endif
 
 	/** Return the PresetSaveDir path relative to the project. */
-	FString GetResolvedPresetSaveDir() const;
+	TAKERECORDER_API FString GetResolvedPresetSaveDir() const;
 
 	/** Sets the default location for the present saved directory. */
-	void SetPresetSaveDir(const FString& InPath);
+	TAKERECORDER_API void SetPresetSaveDir(const FString& InPath);
 
 	/** User settings that should be passed to a recorder instance */
 	UPROPERTY(config, EditAnywhere, BlueprintReadOnly, Category="User Settings", meta=(ShowOnlyInnerProperties))
diff --git a/Engine/Plugins/VirtualProduction/Takes/Source/TakeRecorder/TakeRecorder.Build.cs b/Engine/Plugins/VirtualProduction/Takes/Source/TakeRecorder/TakeRecorder.Build.cs
--- a/Engine/Plugins/VirtualProduction/Takes/Source/TakeRecorder/TakeRecorder.Build.cs
+++ b/Engine/Plugins/VirtualProduction/Takes/Source/TakeRecorder/TakeRecorder.Build.cs
@@ -9,47 +9,31 @@ public class TakeRecorder : ModuleRules
 		PrivateDependencyModuleNames.AddRange(
 			new string[] {
 				"AppFramework",
-				"AssetDefinition",
 				"AssetRegistry",
-				"ClassViewer",
-				"ContentBrowser",
 				"Core",
 				"CoreUObject",
-				"EditorStyle",
-				"EditorWidgets",
 				"Engine",
 				"InputCore",
-				"LevelEditor",
 				"LevelSequence",
 				"MovieScene",
-				"MovieSceneTools",
-				"Projects",
-				"PropertyEditor",
 				"TakesCore",
 				"TakeMovieScene",
 				"TimeManagement",
 				"Settings",
 				"Slate",
 				"SlateCore",
-				"EditorFramework",
-				"UnrealEd",
-				"WorkspaceMenuStructure",
 				"Analytics",
 				"ToolMenus",
 				"ToolWidgets",
 				"NamingTokens",
-				"SequencerCore",
-				"NamingTokensUI",
 			}
 		);
 
 		PublicDependencyModuleNames.AddRange(
 			new string[] {
-				"EditorSubsystem",
 				"UMG",
                 "TakeTrackRecorders",
                 "SerializedRecorderInterface",
-                "Sequencer",
 				"TraceLog"
             }
         );
@@ -60,5 +44,33 @@ public class TakeRecorder : ModuleRules
 				"TakeRecorderNamingTokens"
 			}
 		);
+
+		if (Target.bBuildEditor == true)
+		{
+			PrivateDependencyModuleNames.AddRange(
+				new string[] {
+					"AssetDefinition",
+					"ContentBrowser",
+					"ClassViewer",
+					"EditorStyle",
+					"EditorWidgets",
+					"LevelEditor",
+					"MovieSceneTools",
+					"PropertyEditor",
+					"EditorFramework",
+					"UnrealEd",
+					"WorkspaceMenuStructure",
+					"SequencerCore",
+					"NamingTokensUI",
+				}
+			);
+
+			PublicDependencyModuleNames.AddRange(
+				new string[] {
+					"EditorSubsystem",
+					"Sequencer"
+				}
+			);
+		}
 	}
 }
diff --git a/Engine/Plugins/VirtualProduction/Takes/Takes.uplugin b/Engine/Plugins/VirtualProduction/Takes/Takes.uplugin
--- a/Engine/Plugins/VirtualProduction/Takes/Takes.uplugin
+++ b/Engine/Plugins/VirtualProduction/Takes/Takes.uplugin
@@ -55,6 +55,12 @@
 			"ProgramAllowList": [ "LiveLinkHub" ],
 			"LoadingPhase": "PostDefault"
 		},
+		{
+			"Name": "TakeRecorderEditor",
+			"Type": "UncookedOnly",
+			"ProgramAllowList": [ "LiveLinkHub" ],
+			"LoadingPhase": "PostDefault"
+		},
 		{
 			"Name": "TakeMovieScene",
 			"Type": "Runtime",
