Take Recorder - Make the TakeTrackRecorders module runtime ready.

#rnx

#Synced-CL 52679060

diff --git a/Engine/Plugins/VirtualProduction/Takes/Source/TakeTrackRecorders/Private/MovieScene3DAttachTrackRecorder.cpp b/Engine/Plugins/VirtualProduction/Takes/Source/TakeTrackRecorders/Private/MovieScene3DAttachTrackRecorder.cpp
--- a/Engine/Plugins/VirtualProduction/Takes/Source/TakeTrackRecorders/Private/MovieScene3DAttachTrackRecorder.cpp
+++ b/Engine/Plugins/VirtualProduction/Takes/Source/TakeTrackRecorders/Private/MovieScene3DAttachTrackRecorder.cpp
@@ -4,9 +4,9 @@
 #include "Tracks/MovieScene3DAttachTrack.h"
 #include "Sections/MovieScene3DAttachSection.h"
 #include "Modules/ModuleManager.h"
-#include "SequenceRecorderUtils.h"
 #include "MovieScene.h"
 #include "LevelSequence.h"
+#include "TakeTrackRecordersUtils.h"
 
 #include UE_INLINE_GENERATED_CPP_BY_NAME(MovieScene3DAttachTrackRecorder)
 
@@ -36,7 +36,7 @@ void UMovieScene3DAttachTrackRecorder::RecordSampleImpl(const FQualifiedFrameTim
 		// get attachment and check if the actor we are attached to is being recorded
 		FName SocketName;
 		FName ComponentName;
-		AActor* AttachedToActor = SequenceRecorderUtils::GetAttachment(ActorToRecord, SocketName, ComponentName);
+		AActor* AttachedToActor = UE::TakeTrackRecordersUtils::Private::GetAttachment(ActorToRecord, SocketName, ComponentName);
 
 		if (AttachedToActor && OwningTakeRecorderSource->IsOtherActorBeingRecorded(AttachedToActor))
 		{
@@ -84,15 +84,21 @@ void UMovieScene3DAttachTrackRecorder::FinalizeTrackImpl()
 	{
 		FName SocketName;
 		FName ComponentName;
-		AActor* AttachedToActor = SequenceRecorderUtils::GetAttachment(ActorToRecord, SocketName, ComponentName);
+		AActor* AttachedToActor = UE::TakeTrackRecordersUtils::Private::GetAttachment(ActorToRecord, SocketName, ComponentName);
 		//note this actor may no longer exist BUT we need to do this in finalize since the compilation only happens there.
 		//fix would be to have a way to get the sequence id differently then doing the compilation above
 		//or have another way to do cleanup
 		if (MovieSceneSection.IsValid() && AttachedToActor)
 		{
 			if (!Guid.IsValid())
-			{
-				UE_LOGF(LogTemp, Warning, "Could not find binding to attach (%ls) to its parent (%ls), perhaps (%ls) was not recorded?", *ActorToRecord->GetActorLabel(), *AttachedToActor->GetActorLabel(), *AttachedToActor->GetActorLabel());
+			{				
+			#if WITH_EDITOR
+				const FString ActorId = ActorToRecord->GetActorLabel();
+			#else // WITH_EDITOR
+				const FString ActorId = ActorToRecord->GetName();
+			#endif // WITH_EDITOR
+
+				UE_LOGF(LogTemp, Warning, "Could not find binding to attach (%ls) to its parent (%ls), perhaps (%ls) was not recorded?", *ActorId, *ActorId, *ActorId);
 			}
 
 			FMovieSceneSequenceID TargetSequenceID = OwningTakeRecorderSource->GetLevelSequenceID(AttachedToActor);
diff --git a/Engine/Plugins/VirtualProduction/Takes/Source/TakeTrackRecorders/Private/MovieScene3DTransformTrackRecorder.cpp b/Engine/Plugins/VirtualProduction/Takes/Source/TakeTrackRecorders/Private/MovieScene3DTransformTrackRecorder.cpp
--- a/Engine/Plugins/VirtualProduction/Takes/Source/TakeTrackRecorders/Private/MovieScene3DTransformTrackRecorder.cpp
+++ b/Engine/Plugins/VirtualProduction/Takes/Source/TakeTrackRecorders/Private/MovieScene3DTransformTrackRecorder.cpp
@@ -13,9 +13,9 @@
 #include "Channels/MovieSceneChannelProxy.h"
 #include "Components/SceneComponent.h"
 #include "Components/SkeletalMeshComponent.h"
-#include "SequenceRecorderUtils.h"
 #include "Animation/AnimData/IAnimationDataModel.h"
 #include "Evaluation/MovieSceneEvaluationState.h"
+#include "TakeTrackRecordersUtils.h"
 
 #include UE_INLINE_GENERATED_CPP_BY_NAME(MovieScene3DTransformTrackRecorder)
 
@@ -125,7 +125,7 @@ void UMovieScene3DTransformTrackRecorder::FinalizeTrackImpl()
 	AActor* ActorToRecord = Cast<AActor>(ObjectToRecord.Get());
 	if (ActorToRecord)
 	{
-		AActor* AttachedToActor = SequenceRecorderUtils::GetAttachment(ActorToRecord, SocketName, ComponentName);
+		AActor* AttachedToActor = UE::TakeTrackRecordersUtils::Private::GetAttachment(ActorToRecord, SocketName, ComponentName);
 		//If attached to an actor that has an initial root transform(skeletal mesh) that we negated out when moving over the anim keys to the transform
 		//track we need to reapply that inverse here. But if it's attached to a socket we don't since we are in not in actor space but sill in component space.
 		if (AttachedToActor  && SocketName == NAME_None)
@@ -483,8 +483,11 @@ void UMovieScene3DTransformTrackRecorder::PostProcessAnimationData(UMovieSceneAn
 	}
 	if (AnimTrackRecorder->RootWasRemoved())
 	{
+		// Only accessing AnimationKeys at editor time currently.
+	#if WITH_EDITOR
 		//Get All Animation Keys
 		UE::TrackRecorders::FBufferedTransformKeys  AnimationKeys;
+	#endif // WITH_EDITOR
 
 		USkeletalMeshComponent* SkeletalMeshComponent = AnimTrackRecorder->GetSkeletalMeshComponent();
 		bool bLayerCurrentBuffered = false;
@@ -518,7 +521,11 @@ void UMovieScene3DTransformTrackRecorder::PostProcessAnimationData(UMovieSceneAn
 				USkeleton* AnimSkeleton = AnimSequence->GetSkeleton();
 
 				TArray<FName> TrackNames;
+
+			#if WITH_EDITOR
+				// Need to revisit this when implementing the runtime implementation.
 				AnimSequence->GetDataModelInterface()->GetBoneTrackNames(TrackNames);
+			#endif // WITH_EDITOR
 
 				FName RootBoneName = NAME_None;
 				for (const FName& TrackName : TrackNames)
@@ -550,6 +557,9 @@ void UMovieScene3DTransformTrackRecorder::PostProcessAnimationData(UMovieSceneAn
 				const FFrameNumber StartTime = MovieSceneSection->GetInclusiveStartFrame();
 
 				FTransform InvComponentTransform = AnimTrackRecorder->GetComponentTransform().Inverse();
+
+			#if WITH_EDITOR
+				// Need to revisit this when implementing the runtime implementation.
 				for (int32 KeyIndex = 0; KeyIndex < AnimSequence->GetDataModelInterface()->GetNumberOfKeys(); KeyIndex++)
 				{
 					const FTransform Transform = AnimSequence->GetDataModelInterface()->GetBoneTrackTransform(RootBoneName, FFrameNumber(KeyIndex));
@@ -561,9 +571,12 @@ void UMovieScene3DTransformTrackRecorder::PostProcessAnimationData(UMovieSceneAn
 						DefaultTransform = Total;
 					}
 				}
+			#endif // WITH_EDITOR
 			}
 		}
 
+		// As we only add to AnimationKeys in a WITH_EDITOR block, only access it from WITH_EDITOR as well.
+	#if WITH_EDITOR
 		if (!bLayerCurrentBuffered)
 		{
 			BufferedTransforms = MoveTemp(AnimationKeys);
@@ -573,6 +586,7 @@ void UMovieScene3DTransformTrackRecorder::PostProcessAnimationData(UMovieSceneAn
 			//The current BufferedTransforms need to layer on top of the animation keys
 			BufferedTransforms = AnimationKeys.Collapse(BufferedTransforms);
 		}
+	#endif // WITH_EDITOR
 	}
 }
 
diff --git a/Engine/Plugins/VirtualProduction/Takes/Source/TakeTrackRecorders/Private/MovieSceneActorReferenceTrackRecorder.cpp b/Engine/Plugins/VirtualProduction/Takes/Source/TakeTrackRecorders/Private/MovieSceneActorReferenceTrackRecorder.cpp
--- a/Engine/Plugins/VirtualProduction/Takes/Source/TakeTrackRecorders/Private/MovieSceneActorReferenceTrackRecorder.cpp
+++ b/Engine/Plugins/VirtualProduction/Takes/Source/TakeTrackRecorders/Private/MovieSceneActorReferenceTrackRecorder.cpp
@@ -7,7 +7,6 @@
 #include "LevelSequence.h"
 #include "MovieScene.h"
 #include "MovieSceneSpawnableAnnotation.h"
-#include "SequencerUtilities.h"
 #include "TakeRecorderSource.h"
 #include "TakeRecorderSources.h"
 #include "TakesCoreLog.h"
diff --git a/Engine/Plugins/VirtualProduction/Takes/Source/TakeTrackRecorders/Private/MovieSceneAnimationTrackRecorder.cpp b/Engine/Plugins/VirtualProduction/Takes/Source/TakeTrackRecorders/Private/MovieSceneAnimationTrackRecorder.cpp
--- a/Engine/Plugins/VirtualProduction/Takes/Source/TakeTrackRecorders/Private/MovieSceneAnimationTrackRecorder.cpp
+++ b/Engine/Plugins/VirtualProduction/Takes/Source/TakeTrackRecorders/Private/MovieSceneAnimationTrackRecorder.cpp
@@ -2,7 +2,6 @@
 
 #include "TrackRecorders/MovieSceneAnimationTrackRecorder.h"
 #include "Animation/Skeleton.h"
-#include "AnimationRecorder.h"
 #include "AssetRegistry/AssetData.h"
 #include "Engine/SkeletalMesh.h"
 #include "TrackRecorders/MovieSceneAnimationTrackRecorderSettings.h"
@@ -12,25 +11,32 @@
 #include "Sections/MovieSceneSkeletalAnimationSection.h"
 #include "MovieScene.h"
 #include "AssetRegistry/AssetRegistryModule.h"
-#include "SequenceRecorderSettings.h"
-#include "SequenceRecorderUtils.h"
 #include "Components/SkeletalMeshComponent.h"
 #include "Animation/AnimSequence.h"
+#include "Animation/AnimationRecordingSettings.h"
 #include "Engine/TimecodeProvider.h"
 #include "Engine/Engine.h"
 #include "LevelSequence.h"
 #include "UObject/SavePackage.h"
 #include "UObject/UObjectBaseUtility.h"
-#include "ObjectTools.h"
-#include "Editor.h"
 #include "NamingTokenData.h"
 #include "NamingTokens/TakeRecorderNamingTokensContext.h"
 
+#if WITH_EDITOR
+#include "AnimationRecorder.h"
+#include "SequenceRecorderSettings.h"
+#include "SequenceRecorderUtils.h"
+#include "ObjectTools.h"
+#include "Editor.h"
+#endif // WITH_EDITOR
+
 #include UE_INLINE_GENERATED_CPP_BY_NAME(MovieSceneAnimationTrackRecorder)
 
 #define LOCTEXT_NAMESPACE "MovieSceneAnimationTrackRecorder"
 
+#if WITH_EDITOR
 DEFINE_LOG_CATEGORY(AnimationSerialization);
+#endif // WITH_EDITOR
 
 bool FMovieSceneAnimationTrackRecorderFactory::CanRecordObject(UObject* InObjectToRecord) const
 {
@@ -58,12 +64,17 @@ void UMovieSceneAnimationTrackRecorder::CreateAnimationAssetAndSequence(const AA
 	if (SkeletalMesh.IsValid())
 	{
 		ComponentTransform = SkeletalMeshComponent->GetComponentToWorld().GetRelativeTransform(Actor->GetTransform());
+	#if WITH_EDITOR
 		FString AnimationAssetName = Actor->GetActorLabel();
+	#else // WITH_EDITOR
+		FString AnimationAssetName = Actor->GetName();
+	#endif // WITH_EDITOR
 
 		FString DirectoryPath = AnimationDirectory.Path;
 
 		if (ULevelSequence* LeaderLevelSequence = OwningTakeRecorderSource->GetRootLevelSequence())
 		{
+		#if WITH_EDITOR
 			UTakeMetaData* AssetMetaData = LeaderLevelSequence->FindMetaData<UTakeMetaData>();
 
 			UTakeRecorderNamingTokensContext* Context = NewObject<UTakeRecorderNamingTokensContext>();
@@ -71,6 +82,7 @@ void UMovieSceneAnimationTrackRecorder::CreateAnimationAssetAndSequence(const AA
 			
 			AnimationAssetName = AssetMetaData->GenerateAssetPath(AnimSettings->AnimationAssetName, Context);
 			DirectoryPath = AssetMetaData->GenerateAssetPath(DirectoryPath, Context);
+		#endif // WITH_EDITOR
 		}
 
 		AnimSequence = TakesUtils::MakeNewAsset<UAnimSequence>(DirectoryPath, AnimationAssetName);
@@ -89,7 +101,11 @@ void UMovieSceneAnimationTrackRecorder::CreateAnimationAssetAndSequence(const AA
 // todo move to FTakeUtils?
 static FGuid GetActorInSequence(AActor* InActor, UMovieScene* MovieScene)
 {
+#if WITH_EDITOR
 	FString ActorTargetName = InActor->GetActorLabel();
+#else // WITH_EDITOR
+	FString ActorTargetName = InActor->GetName();
+#endif // WITH_EDITOR
 
 	for (int32 PossessableCount = 0; PossessableCount < MovieScene->GetPossessableCount(); ++PossessableCount)
 	{
@@ -172,8 +188,10 @@ void UMovieSceneAnimationTrackRecorder::CreateTrackImpl()
 
 			if (!AnimationSerializer.OpenForWrite(FileName, Header, Error))
 			{
+			#if WITH_EDITOR
 				//UE_LOGF(LogFrameTransport, Error, "Cannot open frame debugger cache %ls. Failed to create archive.", *InFilename);
 				UE_LOGF(AnimationSerialization, Warning, "Error Opening Animation Sequencer File: Object '%ls' Error '%ls'", *(Name), *(Error.ToString()));
+			#endif // WITH_EDITOR
 			}
 			bAnimationRecorderCreated = false;
 
@@ -192,15 +210,20 @@ void UMovieSceneAnimationTrackRecorder::CreateTrackImpl()
 				FText AnimationTrackDisplayName = AnimSettings->AnimationTrackName;
 				if (const ULevelSequence* LeaderLevelSequence = OwningTakeRecorderSource->GetRootLevelSequence())
 				{
+				#if WITH_EDITOR
 					const UTakeMetaData* AssetMetaData = LeaderLevelSequence->FindMetaData<UTakeMetaData>();
 
 					UTakeRecorderNamingTokensContext* Context = NewObject<UTakeRecorderNamingTokensContext>();
 					Context->Actor = Actor;
 			
 					AnimationTrackDisplayName = FText::FromString(AssetMetaData->GenerateAssetPath(AnimationTrackDisplayName.ToString(), Context));
+				#endif // WITH_EDITOR
 				}
 				
+			#if WITH_EDITOR
 				AnimTrack->SetDisplayName(AnimationTrackDisplayName);
+			#endif // WITH_EDITOR
+
 				AnimTrack->AddNewAnimation(FFrameNumber(0), AnimSequence.Get());
 				MovieSceneSection = Cast<UMovieSceneSkeletalAnimationSection>(AnimTrack->GetAllSections()[0]);
 				MovieSceneSection->Params.bForceCustomMode = true;
@@ -214,10 +237,15 @@ void UMovieSceneAnimationTrackRecorder::StopRecordingImpl()
 {
 	AnimationSerializer.Close();
 
-	// Legacy Animation Recorder allowed recording into an animation asset directly and not creating an movie section
-	const bool bShowAnimationAssetCreatedToast = false;
-	InitialRootTransform = AnimationRecorder.Recorder.Get()->GetInitialRootTransform();
-	AnimationRecorder.FinishRecording(bShowAnimationAssetCreatedToast);
+#if WITH_EDITOR
+	if (AnimationRecorder.IsValid())
+	{
+		// Legacy Animation Recorder allowed recording into an animation asset directly and not creating an movie section
+		const bool bShowAnimationAssetCreatedToast = false;
+		InitialRootTransform = AnimationRecorder->GetRecorder().Get()->GetInitialRootTransform();
+		AnimationRecorder->FinishRecording(bShowAnimationAssetCreatedToast);
+	}
+#endif // WITH_EDITOR
 }
 
 void UMovieSceneAnimationTrackRecorder::FinalizeTrackImpl()
@@ -245,14 +273,18 @@ void UMovieSceneAnimationTrackRecorder::CancelTrackImpl()
 		AssetsToCleanUp.Add(Anim);
 	}
 	
+	// Revisit this cleanup as it still needs to happen in runtime so will need to update this.
+#if WITH_EDITOR
 	if (GEditor && AssetsToCleanUp.Num() > 0)
 	{
 		ObjectTools::ForceDeleteObjects(AssetsToCleanUp, false);
 	}
+#endif // WITH_EDITOR
 }
 
 void UMovieSceneAnimationTrackRecorder::RecordSampleImpl(const FQualifiedFrameTime& CurrentTime)
 {
+#if WITH_EDITOR
 	// The animation recorder does most of the work here
 	//  Note we wait for first tick so that we can make sure all of the attach tracks are set up .
 	float CurrentSeconds = CurrentTime.AsSeconds();
@@ -263,7 +295,12 @@ void UMovieSceneAnimationTrackRecorder::RecordSampleImpl(const FQualifiedFrameTi
 		//Reset the start times based upon when the animation really starts.
 		if (MovieSceneSection.IsValid())
 		{
+		#if WITH_EDITOR
 			MovieSceneSection->TimecodeSource = SequenceRecorderUtils::GetTimecodeSource();
+		#else // WITH_EDITOR
+			// This mirrors what currently happens in SequenceRecorderUtils::GetTimecodeSource.
+			MovieSceneSection->TimecodeSource = FMovieSceneTimecodeSource(FApp::GetTimecode());
+		#endif // WITH_EDITOR
 			FFrameRate   TickResolution = MovieSceneSection->GetTypedOuter<UMovieScene>()->GetTickResolution();
 			FFrameNumber CurrentFrame = CurrentTime.ConvertTo(TickResolution).FloorToFrame();
 
@@ -315,13 +352,16 @@ void UMovieSceneAnimationTrackRecorder::RecordSampleImpl(const FQualifiedFrameTi
 		RecordingSettings.bCheckDeltaTimeAtBeginning = false;
 		RecordingSettings.IncludeAnimationNames = TrackRecorderSettings.IncludeAnimationNames;
 		RecordingSettings.ExcludeAnimationNames = TrackRecorderSettings.ExcludeAnimationNames;
-		AnimationRecorder.Init(SkeletalMeshComponent.Get(), AnimSequence.Get(), &AnimationSerializer, RecordingSettings);
-		AnimationRecorder.BeginRecording();
+
+		AnimationRecorder = MakeShared<FAnimRecorderInstance>();
+
+		AnimationRecorder->Init(SkeletalMeshComponent.Get(), AnimSequence.Get(), &AnimationSerializer, RecordingSettings);
+		AnimationRecorder->BeginRecording();
 	}
 	else
 	{
 		float DeltaTime = CurrentSeconds - PreviousSeconds;
-		AnimationRecorder.Update(DeltaTime);
+		AnimationRecorder->Update(DeltaTime);
 	}
 
 	PreviousSeconds = CurrentSeconds;
@@ -333,6 +373,7 @@ void UMovieSceneAnimationTrackRecorder::RecordSampleImpl(const FQualifiedFrameTi
 		SkeletalMeshComponent->bEnableUpdateRateOptimizations = false;
 		SkeletalMeshComponent->VisibilityBasedAnimTickOption = EVisibilityBasedAnimTickOption::AlwaysTickPoseAndRefreshBones;
 	}
+#endif // WITH_EDITOR
 }
 
 void UMovieSceneAnimationTrackRecorder::RemoveRootMotion()
@@ -366,9 +407,13 @@ void UMovieSceneAnimationTrackRecorder::ProcessRecordedTimes(const FString& Hour
 }
 void UMovieSceneAnimationTrackRecorder::ProcessRecordedTimes(const FProcessRecordedTimeParams& InParams)
 {
-	UMovieSceneAnimationTrackRecorderSettings* AnimSettings = CastChecked<UMovieSceneAnimationTrackRecorderSettings>(Settings.Get());
-
-	AnimationRecorder.ProcessRecordedTimes(AnimSequence.Get(), SkeletalMeshComponent.Get(), AnimSettings->TimecodeBoneMethod, InParams);
+#if WITH_EDITOR
+	if (AnimationRecorder.IsValid())
+	{
+		UMovieSceneAnimationTrackRecorderSettings* AnimSettings = CastChecked<UMovieSceneAnimationTrackRecorderSettings>(Settings.Get());
+		AnimationRecorder->ProcessRecordedTimes(AnimSequence.Get(), SkeletalMeshComponent.Get(), AnimSettings->TimecodeBoneMethod, InParams);
+	}
+#endif // WITH_EDITOR
 }
 
 bool UMovieSceneAnimationTrackRecorder::LoadRecordedFile(const FString& FileName, UMovieScene *InMovieScene, TMap<FGuid, AActor*>& ActorGuidToActorMap,  TFunction<void()> InCompletionCallback)
@@ -416,6 +461,9 @@ bool UMovieSceneAnimationTrackRecorder::LoadRecordedFile(const FString& FileName
 
 								CreateAnimationAssetAndSequence(Actor,AnimationDirectory);
 
+								// Recording the curve data is only possible at editor time and not at runtime. As this is part of a load,
+								// which is not going to be something that will be supported at runtime, for now this can simply be guarded.
+							#if WITH_EDITOR
 								AnimSequence->DeleteNotifyTrackData();
 
 								IAnimationDataController& Controller = AnimSequence->GetController();
@@ -459,6 +507,7 @@ bool UMovieSceneAnimationTrackRecorder::LoadRecordedFile(const FString& FileName
 
 									Controller.NotifyPopulated();
 								}
+							#endif // WITH_EDITOR
 
 								AnimSequence->MarkPackageDirty();
 								FFrameRate TickResolution = InMovieScene->GetTickResolution();;
diff --git a/Engine/Plugins/VirtualProduction/Takes/Source/TakeTrackRecorders/Private/MovieSceneAnimationTrackRecorderSettings.cpp b/Engine/Plugins/VirtualProduction/Takes/Source/TakeTrackRecorders/Private/MovieSceneAnimationTrackRecorderSettings.cpp
--- a/Engine/Plugins/VirtualProduction/Takes/Source/TakeTrackRecorders/Private/MovieSceneAnimationTrackRecorderSettings.cpp
+++ b/Engine/Plugins/VirtualProduction/Takes/Source/TakeTrackRecorders/Private/MovieSceneAnimationTrackRecorderSettings.cpp
@@ -21,6 +21,7 @@ UMovieSceneAnimationTrackRecorderEditorSettings::UMovieSceneAnimationTrackRecord
 	TimecodeBoneMethod.BoneMode = ETimecodeBoneMode::Root;
 }
 
+#if WITH_EDITOR
 void UMovieSceneAnimationTrackRecorderEditorSettings::PostEditChangeProperty(FPropertyChangedEvent& PropertyChangedEvent)
 {
 	Super::PostEditChangeProperty(PropertyChangedEvent);
@@ -30,4 +31,5 @@ void UMovieSceneAnimationTrackRecorderEditorSettings::PostEditChangeProperty(FPr
 		SaveConfig();
 	}
 }
+#endif // WITH_EDITOR
 
diff --git a/Engine/Plugins/VirtualProduction/Takes/Source/TakeTrackRecorders/Private/MovieSceneTrackPropertyRecorder.cpp b/Engine/Plugins/VirtualProduction/Takes/Source/TakeTrackRecorders/Private/MovieSceneTrackPropertyRecorder.cpp
--- a/Engine/Plugins/VirtualProduction/Takes/Source/TakeTrackRecorders/Private/MovieSceneTrackPropertyRecorder.cpp
+++ b/Engine/Plugins/VirtualProduction/Takes/Source/TakeTrackRecorders/Private/MovieSceneTrackPropertyRecorder.cpp
@@ -21,7 +21,6 @@
 #include "Sections/MovieSceneStringSection.h"
 #include "Tracks/MovieSceneStringTrack.h"
 #include "Channels/MovieSceneChannelProxy.h"
-#include "SequenceRecorderUtils.h"
 
 // current set of compiled-in property types
 
diff --git a/Engine/Plugins/VirtualProduction/Takes/Source/TakeTrackRecorders/Private/TakeTrackRecordersUtils.cpp b/Engine/Plugins/VirtualProduction/Takes/Source/TakeTrackRecorders/Private/TakeTrackRecordersUtils.cpp
new file mode 100644
--- /dev/null
+++ b/Engine/Plugins/VirtualProduction/Takes/Source/TakeTrackRecorders/Private/TakeTrackRecordersUtils.cpp
@@ -0,0 +1,31 @@
+// Copyright Epic Games, Inc. All Rights Reserved.
+
+#include "TakeTrackRecordersUtils.h"
+
+#include "GameFramework/Actor.h"
+#include "Components/SceneComponent.h"
+#include "UObject/NameTypes.h"
+
+namespace UE::TakeTrackRecordersUtils::Private
+{
+
+AActor* GetAttachment(AActor* InActor, FName& SocketName, FName& ComponentName)
+{
+	AActor* AttachedActor = nullptr;
+	SocketName = NAME_None;
+	ComponentName = NAME_None;
+	if (InActor)
+	{
+		USceneComponent* RootComponent = InActor->GetRootComponent();
+		if (RootComponent && RootComponent->GetAttachParent() != nullptr)
+		{
+			AttachedActor = RootComponent->GetAttachParent()->GetOwner();
+			SocketName = RootComponent->GetAttachSocketName();
+			ComponentName = RootComponent->GetAttachParent()->GetFName();
+		}
+	}
+
+	return AttachedActor;
+}
+
+} // namespace UE::TakeTrackRecordersUtils::Private
diff --git a/Engine/Plugins/VirtualProduction/Takes/Source/TakeTrackRecorders/Private/TakeTrackRecordersUtils.h b/Engine/Plugins/VirtualProduction/Takes/Source/TakeTrackRecorders/Private/TakeTrackRecordersUtils.h
new file mode 100644
--- /dev/null
+++ b/Engine/Plugins/VirtualProduction/Takes/Source/TakeTrackRecorders/Private/TakeTrackRecordersUtils.h
@@ -0,0 +1,22 @@
+// Copyright Epic Games, Inc. All Rights Reserved.
+
+#pragma once
+
+class AActor;
+class FName;
+
+namespace UE::TakeTrackRecordersUtils::Private
+{
+
+/**
+ * Get the attachment owner of the given actor.
+ * 
+ * @param InActor The actor to find the attachment for.
+ * @param SocketName Out name of the socket of the attachment.
+ * @param ComponentName Out name of the component of the attachment. 
+ * 
+ * @return The attach owning actor.
+ */
+AActor* GetAttachment(AActor* InActor, FName& SocketName, FName& ComponentName);
+
+} // namespace UE::TakeTrackRecordersUtils::Private
diff --git a/Engine/Plugins/VirtualProduction/Takes/Source/TakeTrackRecorders/Public/TrackRecorders/MovieSceneAnimationTrackRecorder.h b/Engine/Plugins/VirtualProduction/Takes/Source/TakeTrackRecorders/Public/TrackRecorders/MovieSceneAnimationTrackRecorder.h
--- a/Engine/Plugins/VirtualProduction/Takes/Source/TakeTrackRecorders/Public/TrackRecorders/MovieSceneAnimationTrackRecorder.h
+++ b/Engine/Plugins/VirtualProduction/Takes/Source/TakeTrackRecorders/Public/TrackRecorders/MovieSceneAnimationTrackRecorder.h
@@ -4,13 +4,15 @@
 
 #include "CoreMinimal.h"
 #include "IMovieSceneTrackRecorderFactory.h"
+#include "Templates/SharedPointer.h"
 #include "MovieScene.h"
 #include "MovieSceneTrackRecorder.h"
 #include "MovieSceneAnimationTrackRecorderSettings.h"
 #include "Animation/AnimSequence.h"
 #include "Serializers/MovieSceneAnimationSerialization.h"
 #include "Sections/MovieSceneSkeletalAnimationSection.h"
-#include "AnimationRecorder.h"
+#include "IAnimRecorderInstance.h"
+
 #include "MovieSceneAnimationTrackRecorder.generated.h"
 
 #define UE_API TAKETRACKRECORDERS_API
@@ -94,8 +96,8 @@ private:
 
 	bool bAnimationRecorderCreated;
 
-	/** Animatinon Recorder */
-	FAnimRecorderInstance AnimationRecorder;
+	/** Animation Recorder */
+	TSharedPtr<UE::AnimationRecording::IAnimRecorderInstance> AnimationRecorder;
 
 	/** Previous Seconds to calc Delta used by Animation Recorder **/
 	float  PreviousSeconds;
diff --git a/Engine/Plugins/VirtualProduction/Takes/Source/TakeTrackRecorders/Public/TrackRecorders/MovieSceneAnimationTrackRecorderSettings.h b/Engine/Plugins/VirtualProduction/Takes/Source/TakeTrackRecorders/Public/TrackRecorders/MovieSceneAnimationTrackRecorderSettings.h
--- a/Engine/Plugins/VirtualProduction/Takes/Source/TakeTrackRecorders/Public/TrackRecorders/MovieSceneAnimationTrackRecorderSettings.h
+++ b/Engine/Plugins/VirtualProduction/Takes/Source/TakeTrackRecorders/Public/TrackRecorders/MovieSceneAnimationTrackRecorderSettings.h
@@ -6,7 +6,8 @@
 #include "Misc/FrameRate.h"
 #include "MovieSceneTrackRecorderSettings.h"
 #include "Curves/RichCurve.h"
-#include "AnimationRecorder.h"
+#include "TimecodeBoneMethod.h"
+
 #include "MovieSceneAnimationTrackRecorderSettings.generated.h"
 
 #define UE_API TAKETRACKRECORDERS_API
@@ -18,7 +19,9 @@ class UMovieSceneAnimationTrackRecorderEditorSettings : public UMovieSceneTrackR
 public:
 	UE_API UMovieSceneAnimationTrackRecorderEditorSettings(const FObjectInitializer& ObjInit);
 
+#if WITH_EDITOR
 	UE_API virtual void PostEditChangeProperty(FPropertyChangedEvent& PropertyChangedEvent) override;
+#endif // WITH_EDITOR
 
 	/** Name of the recorded animation track. */
 	UPROPERTY(config, EditAnywhere, BlueprintReadWrite, Category = "Animation Recorder Settings", meta = (NamingTokens))
diff --git a/Engine/Plugins/VirtualProduction/Takes/Source/TakeTrackRecorders/Public/TrackRecorders/MovieSceneTrackPropertyRecorder.h b/Engine/Plugins/VirtualProduction/Takes/Source/TakeTrackRecorders/Public/TrackRecorders/MovieSceneTrackPropertyRecorder.h
--- a/Engine/Plugins/VirtualProduction/Takes/Source/TakeTrackRecorders/Public/TrackRecorders/MovieSceneTrackPropertyRecorder.h
+++ b/Engine/Plugins/VirtualProduction/Takes/Source/TakeTrackRecorders/Public/TrackRecorders/MovieSceneTrackPropertyRecorder.h
@@ -72,7 +72,11 @@ public:
 		}
 		else
 		{
-			FString TrackDisplayName = *Binding.GetProperty(*InObjectToRecord)->GetDisplayNameText().ToString();
+		#if WITH_EDITOR
+			const FString TrackDisplayName = *Binding.GetProperty(*InObjectToRecord)->GetDisplayNameText().ToString();
+		#else // WITH_EDITOR
+			const FString TrackDisplayName = Binding.GetProperty(*InObjectToRecord)->GetName();
+		#endif // WITH_EDITOR
 			MovieSceneSection = AddSection(TrackDisplayName, InMovieScene, InGuid, bOpenSerializer);
 			
 			// Disable the section after creation so that the track can't be evaluated by Sequencer while recording.
diff --git a/Engine/Plugins/VirtualProduction/Takes/Source/TakeTrackRecorders/TakeTrackRecorders.Build.cs b/Engine/Plugins/VirtualProduction/Takes/Source/TakeTrackRecorders/TakeTrackRecorders.Build.cs
--- a/Engine/Plugins/VirtualProduction/Takes/Source/TakeTrackRecorders/TakeTrackRecorders.Build.cs
+++ b/Engine/Plugins/VirtualProduction/Takes/Source/TakeTrackRecorders/TakeTrackRecorders.Build.cs
@@ -15,12 +15,27 @@ public class TakeTrackRecorders : ModuleRules
                 "MovieSceneTracks",
 				"TakesCore",
 				"TimeManagement",
-				"SequenceRecorder", // For access to the individual UProperty type recorders
 				"SequenceRuntimeRecorder", // For FTimecodeBoneMethod and ETimecodeBoneMode
 				"SerializedRecorderInterface",
-				"UnrealEd",
-				"Sequencer"
+				"LevelSequence"
             }
         );
+
+		PublicDependencyModuleNames.AddRange(
+			new string[] {
+				"SequenceRuntimeRecorder"
+			}
+		);
+
+		if (Target.bBuildEditor == true)
+		{
+			PrivateDependencyModuleNames.AddRange(
+				new string[] {
+					"SequenceRecorder",
+					"UnrealEd",
+					"Sequencer"
+				}
+			);
+		}
 	}
 }
