Take Recorder - Final changes to make Takes modules runtimes.

This includes updating Takes.uplugin to change module types to runtime.

The modules that are now runtime are:
TakesCore
TakeRecorder
TakeTrackRecorders
TakeRecorderSources
TakeRecorderNamingTokens

#Synced-CL 52680171

diff --git a/Engine/Plugins/VirtualProduction/Takes/Source/TakeRecorder/Private/Recorder/TakeRecorder.cpp b/Engine/Plugins/VirtualProduction/Takes/Source/TakeRecorder/Private/Recorder/TakeRecorder.cpp
--- a/Engine/Plugins/VirtualProduction/Takes/Source/TakeRecorder/Private/Recorder/TakeRecorder.cpp
+++ b/Engine/Plugins/VirtualProduction/Takes/Source/TakeRecorder/Private/Recorder/TakeRecorder.cpp
@@ -30,6 +30,7 @@
 
 // Engine includes
 #include "GameFramework/WorldSettings.h"
+#include "Framework/Application/SlateApplication.h"
 
 #if WITH_EDITOR
 #include "ISequencer.h"
diff --git a/Engine/Plugins/VirtualProduction/Takes/Source/TakeRecorder/Private/TakeRecorderModule.cpp b/Engine/Plugins/VirtualProduction/Takes/Source/TakeRecorder/Private/TakeRecorderModule.cpp
--- a/Engine/Plugins/VirtualProduction/Takes/Source/TakeRecorder/Private/TakeRecorderModule.cpp
+++ b/Engine/Plugins/VirtualProduction/Takes/Source/TakeRecorder/Private/TakeRecorderModule.cpp
@@ -15,7 +15,6 @@
 #include "Framework/MultiBox/MultiBoxBuilder.h"
 #include "Features/IModularFeatures.h"
 #include "ITakeRecorderDropHandler.h"
-#include "ISettingsModule.h"
 #include "TakeRecorderSettings.h"
 #include "TakePresetSettings.h"
 
@@ -93,7 +92,6 @@ void FTakeRecorderModule::ShutdownModule()
 	FCoreDelegates::GetOnPostEngineInit().RemoveAll(this);
 #endif
 
-	UnregisterSettings();
 	UnregisterSerializedRecorder();
 }
 
@@ -154,20 +152,6 @@ void FTakeRecorderModule::RegisterSettings()
 	RegisterSettingsObject(GetMutableDefault<UMovieSceneTakeSettings>());
 	RegisterSettingsObject(GetMutableDefault<UTakePresetSettings>());
 
-	ISettingsModule& SettingsModule = FModuleManager::LoadModuleChecked<ISettingsModule>("Settings");
-
-	SettingsModule.RegisterSettings("Project", "Plugins", "Take Recorder",
-		LOCTEXT("ProjectSettings_Label", "Take Recorder"),
-		LOCTEXT("ProjectSettings_Description", "Configure project-wide defaults for take recorder."),
-		GetMutableDefault<UTakeRecorderProjectSettings>()
-	);
-
-	SettingsModule.RegisterSettings("Editor", "ContentEditors", "Take Recorder",
-		LOCTEXT("UserSettings_Label", "Take Recorder"),
-		LOCTEXT("UserSettings_Description", "Configure user-specific settings for take recorder."),
-		GetMutableDefault<UTakeRecorderUserSettings>()
-	);
-
 	GetMutableDefault<UTakeRecorderUserSettings>()->LoadConfig();
 	const bool bSaveRecordedAssetsOverride = CVarTakeRecorderSaveRecordedAssetsOverride.GetValueOnGameThread() != 0;
 	if (bSaveRecordedAssetsOverride)
@@ -176,16 +160,6 @@ void FTakeRecorderModule::RegisterSettings()
 	}
 }
 
-void FTakeRecorderModule::UnregisterSettings()
-{
-	ISettingsModule* SettingsModule = FModuleManager::GetModulePtr<ISettingsModule>("Settings");
-	if (SettingsModule)
-	{
-		SettingsModule->UnregisterSettings("Project", "Plugins", "Take Recorder");
-		SettingsModule->UnregisterSettings("Editor",  "ContentEditors", "Take Recorder");
-	}
-}
-
 void FTakeRecorderModule::RegisterSerializedRecorder()
 {
 	SerializedRecorder = MakeShared<FSerializedRecorder>();
diff --git a/Engine/Plugins/VirtualProduction/Takes/Source/TakeRecorder/Private/TakeRecorderModule.h b/Engine/Plugins/VirtualProduction/Takes/Source/TakeRecorder/Private/TakeRecorderModule.h
--- a/Engine/Plugins/VirtualProduction/Takes/Source/TakeRecorder/Private/TakeRecorderModule.h
+++ b/Engine/Plugins/VirtualProduction/Takes/Source/TakeRecorder/Private/TakeRecorderModule.h
@@ -47,7 +47,6 @@ public:
 	
 private:
 	void RegisterSettings();
-	void UnregisterSettings();
 
 	void RegisterSerializedRecorder();
 	void UnregisterSerializedRecorder() const;
diff --git a/Engine/Plugins/VirtualProduction/Takes/Source/TakeRecorder/TakeRecorder.Build.cs b/Engine/Plugins/VirtualProduction/Takes/Source/TakeRecorder/TakeRecorder.Build.cs
--- a/Engine/Plugins/VirtualProduction/Takes/Source/TakeRecorder/TakeRecorder.Build.cs
+++ b/Engine/Plugins/VirtualProduction/Takes/Source/TakeRecorder/TakeRecorder.Build.cs
@@ -19,13 +19,11 @@ public class TakeRecorder : ModuleRules
 				"TakesCore",
 				"TakeMovieScene",
 				"TimeManagement",
-				"Settings",
 				"Slate",
 				"SlateCore",
 				"Analytics",
-				"ToolMenus",
-				"ToolWidgets",
 				"NamingTokens",
+				"Json"
 			}
 		);
 
@@ -62,6 +60,8 @@ public class TakeRecorder : ModuleRules
 					"WorkspaceMenuStructure",
 					"SequencerCore",
 					"NamingTokensUI",
+					"ToolMenus",
+					"ToolWidgets",
 				}
 			);
 
diff --git a/Engine/Plugins/VirtualProduction/Takes/Source/TakeRecorderEditor/Private/TakeRecorderEditorModule.cpp b/Engine/Plugins/VirtualProduction/Takes/Source/TakeRecorderEditor/Private/TakeRecorderEditorModule.cpp
--- a/Engine/Plugins/VirtualProduction/Takes/Source/TakeRecorderEditor/Private/TakeRecorderEditorModule.cpp
+++ b/Engine/Plugins/VirtualProduction/Takes/Source/TakeRecorderEditor/Private/TakeRecorderEditorModule.cpp
@@ -12,6 +12,7 @@
 #include "LevelEditor.h"
 #include "TakeRecorderCommands.h"
 #include "TakeRecorderStyle.h"
+#include "TakeRecorderSettings.h"
 #include "Recorder/TakeRecorder.h"
 
 #include "WorkspaceMenuStructureModule.h"
@@ -528,6 +529,18 @@ void FTakeRecorderEditorModule::RegisterSettings()
 
 	SequencerSettings = USequencerSettingsContainer::GetOrCreate<USequencerSettings>(TEXT("TakeRecorderSequenceEditor"));
 
+	SettingsModule.RegisterSettings("Project", "Plugins", "Take Recorder",
+		LOCTEXT("ProjectSettings_Label", "Take Recorder"),
+		LOCTEXT("ProjectSettings_Description", "Configure project-wide defaults for take recorder."),
+		GetMutableDefault<UTakeRecorderProjectSettings>()
+	);
+
+	SettingsModule.RegisterSettings("Editor", "ContentEditors", "Take Recorder",
+		LOCTEXT("UserSettings_Label", "Take Recorder"),
+		LOCTEXT("UserSettings_Description", "Configure user-specific settings for take recorder."),
+		GetMutableDefault<UTakeRecorderUserSettings>()
+	);
+
 	SettingsModule.RegisterSettings("Editor", "ContentEditors", "TakeRecorderSequenceEditor",
 		LOCTEXT("TakeRecorderSequenceEditorSettingsName", "Take Recorder Sequence Editor"),
 		LOCTEXT("TakeRecorderSequenceEditorSettingsDescription", "Configure the look and feel of the Take Recorder Sequence Editor."),
@@ -539,6 +552,8 @@ void FTakeRecorderEditorModule::UnregisterSettings()
 	ISettingsModule* SettingsModule = FModuleManager::GetModulePtr<ISettingsModule>("Settings");
 	if (SettingsModule)
 	{
+		SettingsModule->UnregisterSettings("Project", "Plugins", "Take Recorder");
+		SettingsModule->UnregisterSettings("Editor", "ContentEditors", "Take Recorder");
 		SettingsModule->UnregisterSettings("Editor", "ContentEditors", "TakeRecorderSequenceEditor");
 	}
 }
diff --git a/Engine/Plugins/VirtualProduction/Takes/Source/TakeRecorderSources/Private/TakeRecorderCameraCutSource.cpp b/Engine/Plugins/VirtualProduction/Takes/Source/TakeRecorderSources/Private/TakeRecorderCameraCutSource.cpp
--- a/Engine/Plugins/VirtualProduction/Takes/Source/TakeRecorderSources/Private/TakeRecorderCameraCutSource.cpp
+++ b/Engine/Plugins/VirtualProduction/Takes/Source/TakeRecorderSources/Private/TakeRecorderCameraCutSource.cpp
@@ -12,6 +12,7 @@
 #include "Engine/Engine.h"
 #include "Camera/PlayerCameraManager.h"
 #include "GameFramework/PlayerController.h"
+#include "Engine/GameInstance.h"
 
 #include "Tracks/MovieSceneCameraCutTrack.h"
 #include "Sections/MovieSceneCameraCutSection.h"
diff --git a/Engine/Plugins/VirtualProduction/Takes/Source/TakeRecorderSources/Private/TakeRecorderSourcesModule.cpp b/Engine/Plugins/VirtualProduction/Takes/Source/TakeRecorderSources/Private/TakeRecorderSourcesModule.cpp
--- a/Engine/Plugins/VirtualProduction/Takes/Source/TakeRecorderSources/Private/TakeRecorderSourcesModule.cpp
+++ b/Engine/Plugins/VirtualProduction/Takes/Source/TakeRecorderSources/Private/TakeRecorderSourcesModule.cpp
@@ -27,6 +27,8 @@
 #include "Framework/Application/SlateApplication.h"
 #include "Framework/MultiBox/MultiBoxBuilder.h"
 #include "Framework/MultiBox/MultiBoxExtender.h"
+#include "Engine/Level.h"
+#include "Engine/GameInstance.h"
 #include "Widgets/Layout/SBox.h"
 #include "Styling/SlateIconFinder.h"
 
diff --git a/Engine/Plugins/VirtualProduction/Takes/Source/TakeRecorderSources/Private/TakeRecorderSourcesModule.h b/Engine/Plugins/VirtualProduction/Takes/Source/TakeRecorderSources/Private/TakeRecorderSourcesModule.h
--- a/Engine/Plugins/VirtualProduction/Takes/Source/TakeRecorderSources/Private/TakeRecorderSourcesModule.h
+++ b/Engine/Plugins/VirtualProduction/Takes/Source/TakeRecorderSources/Private/TakeRecorderSourcesModule.h
@@ -6,6 +6,7 @@
 #include "ITakeRecorderSourcesModule.h"
 #include "Misc/CoreMisc.h"
 #include "Templates/SharedPointer.h"
+#include "Framework/Commands/UICommandList.h"
 
 #if WITH_EDITOR
 #include "LevelEditor.h"
diff --git a/Engine/Plugins/VirtualProduction/Takes/Takes.uplugin b/Engine/Plugins/VirtualProduction/Takes/Takes.uplugin
--- a/Engine/Plugins/VirtualProduction/Takes/Takes.uplugin
+++ b/Engine/Plugins/VirtualProduction/Takes/Takes.uplugin
@@ -29,29 +29,28 @@
 			"Enabled": true
 		}
 	],
-	"Modules": 
-	[
+	"Modules": [
 		{
 			"Name": "TakesCore",
-			"Type": "UncookedOnly",
+			"Type": "Runtime",
 			"ProgramAllowList": [ "LiveLinkHub" ],
 			"LoadingPhase": "PostDefault"
 		},
 		{
 			"Name": "TakeTrackRecorders",
-			"Type": "UncookedOnly",
+			"Type": "Runtime",
 			"ProgramAllowList": [ "LiveLinkHub" ],
 			"LoadingPhase": "PostDefault"
 		},
 		{
 			"Name": "TakeRecorderSources",
 			"ProgramAllowList": [ "LiveLinkHub" ],
-			"Type": "UncookedOnly",
+			"Type": "Runtime",
 			"LoadingPhase": "PostDefault"
 		},
 		{
 			"Name": "TakeRecorder",
-			"Type": "UncookedOnly",
+			"Type": "Runtime",
 			"ProgramAllowList": [ "LiveLinkHub" ],
 			"LoadingPhase": "PostDefault"
 		},
@@ -82,7 +81,7 @@
 		{
 			"Name": "TakeRecorderNamingTokens",
 			"ProgramAllowList": [ "LiveLinkHub" ],
-			"Type": "UncookedOnly",
+			"Type": "Runtime",
 			"LoadingPhase": "PostDefault"
 		}
 	]
